Collab Viewer Template Tweaks

Enabling enterprise plugins (datasmith etc)
Assets editor now points to the Clear Sky environment
Rendering settings aligned with Archviz template

diff --git a/Templates/TP_CollaborativeBP/Config/DefaultEngine.ini b/Templates/TP_CollaborativeBP/Config/DefaultEngine.ini
--- a/Templates/TP_CollaborativeBP/Config/DefaultEngine.ini
+++ b/Templates/TP_CollaborativeBP/Config/DefaultEngine.ini
@@ -1,5 +1,33 @@
-[URL]
-GameName=CollaborationDesignViewer
+[/Script/HardwareTargeting.HardwareTargetingSettings]
+; quality settings
+TargetedHardwareClass=Desktop
+AppliedTargetedHardwareClass=Desktop
+DefaultGraphicsPerformance=Maximum
+AppliedDefaultGraphicsPerformance=Maximum
+
+
+[/Script/Engine.RendererSettings]
+; exposure control and tone mapping settings
+r.Tonemapper.Sharpen=2
+r.TemporalAA.Upsampling=True
+r.DefaultFeature.AutoExposure.ExtendDefaultLuminanceRange=True
+
+;dynamic gi settings
+r.LightPropagationVolume=0
+
+; quality settings
+r.ReflectionCaptureResolution=2048
+r.AllowStaticLighting=True
+r.HighResScreenshotDelay=8
+r.DefaultBackBufferPixelFormat=4
+r.AllowGlobalClipPlane=False
+
+; shader settings
+r.SupportSkyAtmosphereAffectsHeightFog=True
+r.DefaultFeature.LensFlare=True
+r.ClearCoatNormal=False
+r.NormalMapsForStaticLighting=False
+r.CustomDepth=3
 
 [/Script/EngineSettings.GameMapsSettings]
 GameInstanceClass=/Game/CollaborativeViewer/Blueprints/GameInstance/BP_CollaborativeViewer_GameInstance.BP_CollaborativeViewer_GameInstance_C
@@ -7,9 +35,6 @@ GlobalDefaultGameMode=/Game/CollaborativeViewer/Blueprints/GameMode/BP_Collabora
 EditorStartupMap=/Game/CollaborativeViewer/Maps/CollaborativeViewer_P.CollaborativeViewer_P
 GameDefaultMap=/Game/CollaborativeViewer/Maps/MainMenu.MainMenu
 
-[/Script/Engine.RendererSettings]
-r.DefaultFeature.AutoExposure.ExtendDefaultLuminanceRange=True
-r.CustomDepth=3
 
 [/Script/Engine.UserInterfaceSettings]
 UIScaleRule=ShortestSide
@@ -59,50 +84,9 @@ NearClipPlane=1.000000
 +CollisionChannelRedirects=(OldName="VehicleMovement",NewName="Vehicle")
 +CollisionChannelRedirects=(OldName="PawnMovement",NewName="Pawn")
 
-[/Script/Engine.PhysicsSettings]
-DefaultGravityZ=-980.000000
-DefaultTerminalVelocity=4000.000000
-DefaultFluidFriction=0.300000
-SimulateScratchMemorySize=262144
-RagdollAggregateThreshold=4
-TriangleMeshTriangleMinAreaThreshold=5.000000
-bEnableShapeSharing=False
-bEnablePCM=True
-bEnableStabilization=False
-bWarnMissingLocks=True
-bEnable2DPhysics=False
-PhysicErrorCorrection=(PingExtrapolation=0.100000,PingLimit=100.000000,ErrorPerLinearDifference=1.000000,ErrorPerAngularDifference=1.000000,MaxRestoredStateError=1.000000,MaxLinearHardSnapDistance=400.000000,PositionLerp=0.000000,AngleLerp=0.400000,LinearVelocityCoefficient=100.000000,AngularVelocityCoefficient=10.000000,ErrorAccumulationSeconds=0.500000,ErrorAccumulationDistanceSq=15.000000,ErrorAccumulationSimilarity=100.000000)
-LockedAxis=Invalid
-DefaultDegreesOfFreedom=Full3D
-BounceThresholdVelocity=200.000000
-FrictionCombineMode=Average
-RestitutionCombineMode=Average
-MaxAngularVelocity=3600.000000
-MaxDepenetrationVelocity=0.000000
-ContactOffsetMultiplier=0.020000
-MinContactOffset=2.000000
-MaxContactOffset=8.000000
-bSimulateSkeletalMeshOnDedicatedServer=True
-DefaultShapeComplexity=CTF_UseSimpleAndComplex
-bDefaultHasComplexCollision=True
-bSuppressFaceRemapTable=False
-bSupportUVFromHitResults=False
-bDisableActiveActors=False
-bDisableKinematicStaticPairs=False
-bDisableKinematicKinematicPairs=False
-bDisableCCD=False
-bEnableEnhancedDeterminism=False
-AnimPhysicsMinDeltaTime=0.000000
-bSimulateAnimPhysicsAfterReset=False
-MaxPhysicsDeltaTime=0.033333
-bSubstepping=False
-bSubsteppingAsync=False
-MaxSubstepDeltaTime=0.016667
-MaxSubsteps=6
-SyncSceneSmoothingFactor=0.000000
-InitialAverageFrameRate=0.016667
-PhysXTreeRebuildRate=10
-DefaultBroadphaseSettings=(bUseMBPOnClient=False,bUseMBPOnServer=False,bUseMBPOuterBounds=False,MBPBounds=(Min=(X=0.000000,Y=0.000000,Z=0.000000),Max=(X=0.000000,Y=0.000000,Z=0.000000),IsValid=0),MBPOuterBounds=(Min=(X=0.000000,Y=0.000000,Z=0.000000),Max=(X=0.000000,Y=0.000000,Z=0.000000),IsValid=0),MBPNumSubdivs=2)
-ChaosSettings=(DefaultThreadingModel=TaskGraph,DedicatedThreadTickMode=VariableCappedWithTarget,DedicatedThreadBufferMode=Double)
-
-
+[/Script/UnrealEd.UnrealEdEngine]
+; Remove Engine Template maps
+-TemplateMapInfos=(ThumbnailTexture=Texture2D'/Engine/Maps/Templates/Thumbnails/Default.Default',Map="/Engine/Maps/Templates/Template_Default")
+-TemplateMapInfos=(ThumbnailTexture=Texture2D'/Engine/Maps/Templates/Thumbnails/TimeOfDay.TimeOfDay',Map="/Engine/Maps/Templates/TimeOfDay_Default")
+-TemplateMapInfos=(ThumbnailTexture=Texture2D'/Engine/Maps/Templates/Thumbnails/VR-Basic.VR-Basic',Map="/Engine/Maps/Templates/VR-Basic")
++TemplateMapInfos=(ThumbnailTexture=Texture2D'/Engine/Maps/Templates/Thumbnails/Basic.Basic',Map="/Engine/Maps/Templates/Basic")
diff --git a/Templates/TP_CollaborativeBP/TP_CollaborativeBP.uproject b/Templates/TP_CollaborativeBP/TP_CollaborativeBP.uproject
--- a/Templates/TP_CollaborativeBP/TP_CollaborativeBP.uproject
+++ b/Templates/TP_CollaborativeBP/TP_CollaborativeBP.uproject
@@ -2,5 +2,43 @@
 	"FileVersion": 3,
 	"EngineAssociation": "",
 	"Category": "",
-	"Description": ""
+	"Description": "",
+	"Plugins": [
+		{
+			"Name": "SunPosition",
+			"Enabled": true
+		},
+		{
+			"Name": "DataPrepEditor",
+			"Enabled": true
+		},
+		{
+			"Name": "DatasmithImporter",
+			"Enabled": true
+		},
+		{
+			"Name": "DatasmithCADImporter",
+			"Enabled": true
+		},
+		{
+			"Name": "VariantManager",
+			"Enabled": true
+		},
+		{
+			"Name": "PythonScriptPlugin",
+			"Enabled": true
+		},
+		{
+			"Name": "HDRIBackdrop",
+			"Enabled": true
+		},
+		{
+			"Name": "SequencerScripting",
+			"Enabled": false
+		},
+				{
+			"Name": "StaticMeshEditorExtension",
+			"Enabled": true
+		}
+	]
 }
\ No newline at end of file
