Update Collaborative Project template ICON.

diff --git a/Templates/TP_CollaborativeBP/Config/TemplateDefs.ini b/Templates/TP_CollaborativeBP/Config/TemplateDefs.ini
--- a/Templates/TP_CollaborativeBP/Config/TemplateDefs.ini
+++ b/Templates/TP_CollaborativeBP/Config/TemplateDefs.ini
@@ -6,6 +6,7 @@
 
 ; The categories that this template should appear in. The values in this list refer to the Key value given in TemplateCategories.ini.
 Categories=AEC
+Categories=MFG
 
 ; The localized display names for this template. Each entry specifies a Language code and Unicode text to display in the New Project dialog.
 LocalizedDisplayNames=(Language="en", Text="Collab Viewer")
@@ -81,7 +82,7 @@ HiddenSettings=Languages
 HiddenSettings=HardwareTarget
 HiddenSettings=GraphicsPreset 
 ;HiddenSettings=StarterContent 
-;HiddenSettings=XR
+HiddenSettings=XR
 
 
 ; The class types that this project uses.
