Moving Collaborative template in both AEC and MFG
Updating description

diff --git a/Templates/TP_CollaborativeBP/Config/TemplateDefs.ini b/Templates/TP_CollaborativeBP/Config/TemplateDefs.ini
--- a/Templates/TP_CollaborativeBP/Config/TemplateDefs.ini
+++ b/Templates/TP_CollaborativeBP/Config/TemplateDefs.ini
@@ -12,9 +12,9 @@ Categories=MFG
 LocalizedDisplayNames=(Language="en", Text="Collab Viewer")
 
 ; The localized descriptions for this template, which is displayed when the template is selected.
-LocalizedDescriptions=(Language="en", Text="This project template provides navigation and interaction for various platforms (Desktop, VR). Use as a starting point for collaborative Industrial or Architectural projects.")
+LocalizedDescriptions=(Language="en", Text="This project template provides navigation and interaction for various platforms (Desktop, VR). Use as a starting point for collaborative industrial or architectural projects.\n\nThis template enables you to use Datasmith to import models from various design applications and easily light them.")
 ; The sort key to use for the template to bypass the default alphabetical sorting.
-; SortKey="_1"
+SortKey="3"
 
 ; A list of folders to avoid copying when creating a project from this template. 
 ; These should contain files only of concern to the template itself and not desired in the user's project.
