Adding Enterprise templates to generated build and feature packs

diff --git a/Engine/Build/InstalledEngineBuild.xml b/Engine/Build/InstalledEngineBuild.xml
--- a/Engine/Build/InstalledEngineBuild.xml
+++ b/Engine/Build/InstalledEngineBuild.xml
@@ -569,7 +569,7 @@
 		<Node Name="Make Feature Packs" Requires="Update Version Files;Build Tools $(HostPlatform)">
 			<!--Delete any pre-existing feature packs-->
 			<Delete Files="FeaturePacks/..."/>
-			<ForEach Name="TemplateName" Values="FP_FirstPerson;FP_FirstPersonBP;TP_2DSideScroller;TP_2DSideScrollerBP;TP_Flying;TP_FlyingBP;TP_HandheldARBP;TP_Puzzle;TP_PuzzleBP;TP_Rolling;TP_RollingBP;TP_SideScroller;TP_SideScrollerBP;TP_ThirdPerson;TP_ThirdPersonBP;TP_TopDown;TP_TopDownBP;TP_TwinStick;TP_TwinStickBP;TP_Vehicle;TP_VehicleAdv;TP_VehicleAdvBP;TP_VehicleBP;TP_VirtualRealityBP">
+			<ForEach Name="TemplateName" Values="$(ProjectsToFeaturePack)">
 				<PakFile Output="FeaturePacks/$(TemplateName).upack" Files="" ResponseFile="Templates/$(TemplateName)/contents.txt"/>
 			</ForEach>
 			<ForEach Name="SampleName" Values="MobileStarterContent;StarterContent">
diff --git a/Engine/Build/InstalledEngineFilters.xml b/Engine/Build/InstalledEngineFilters.xml
--- a/Engine/Build/InstalledEngineFilters.xml
+++ b/Engine/Build/InstalledEngineFilters.xml
@@ -38,6 +38,35 @@
 	<Property Name="MacSignFilter" Value="*.dylib;*.app"/>
 	<Property Name="SignExceptions" Value=".../ThirdParty/..."/>
 
+	<!-- List of project to build Feature packs -->
+	<Property Name="ProjectsToFeaturePack">
+		FP_FirstPerson
+		FP_FirstPersonBP
+		TP_2DSideScroller
+		TP_2DSideScrollerBP
+		TP_Flying
+		TP_FlyingBP
+		TP_HandheldARBP
+		TP_Puzzle
+		TP_PuzzleBP
+		TP_Rolling
+		TP_RollingBP
+		TP_SideScroller
+		TP_SideScrollerBP
+		TP_ThirdPerson
+		TP_ThirdPersonBP
+		TP_TopDown
+		TP_TopDownBP
+		TP_TwinStick
+		TP_TwinStickBP
+		TP_Vehicle
+		TP_VehicleAdv
+		TP_VehicleAdvBP
+		TP_VehicleBP
+		TP_VirtualRealityBP
+		TP_CollaborativeBP
+	</Property>
+
 	<!-- List of projects to build DDC for -->
 	<Property Name="ProjectsToBuildDDC">
 		Templates/FP_FirstPersonBP/FP_FirstPersonBP.uproject
@@ -52,6 +81,7 @@
 		Templates/TP_PuzzleBP/TP_PuzzleBP.uproject
 		Templates/TP_2DSideScrollerBP/TP_2DSideScrollerBP.uproject
 		Templates/TP_VehicleAdvBP/TP_VehicleAdvBP.uproject
+		Templates/TP_CollaborativeBP/TP_CollaborativeBP.uproject
 		Samples/StarterContent/StarterContent.uproject
 		Samples/MobileStarterContent/MobileStarterContent.uproject
 	</Property>
@@ -117,9 +147,6 @@
 		Templates/Media/...
 		Templates/TP_Blank/...
 		Templates/TP_BlankBP/...
-		Templates/TP_AEC_BlankBP/...
-		Templates/TP_ME_BlankBP/...
-		Templates/TP_MFG_BlankBP/...
 		Templates/TP_FirstPerson/...
 		Templates/TP_FirstPersonBP/...
 		Templates/TP_Flying/...
@@ -144,6 +171,16 @@
 		Templates/TP_VehicleAdv/...
 		Templates/TP_VehicleAdvBP/...
 
+		<!-- Enterprise Templates -->
+		Templates/TP_AEC_BlankBP/...
+		Templates/TP_MFG_BlankBP/...
+		Templates/TP_CollaborativeBP/...
+		Templates/TP_ArchvisBP/...
+		Templates/TP_PhotoStudioBP/...
+
+		Templates/TP_ME_BlankBP/...
+		Templates/TP_ME_VProdBP/...
+
 		<!-- Shared template resources -->
 		Templates/TemplateResources/...
 
diff --git a/Templates/TP_CollaborativeBP/Config/config.ini b/Templates/TP_CollaborativeBP/Config/config.ini
new file mode 100644
--- /dev/null
+++ b/Templates/TP_CollaborativeBP/Config/config.ini
@@ -0,0 +1,119 @@
+[InputSettings]
++AxisConfig=(AxisKeyName="OculusTouch_Right_Thumbstick",AxisProperties=(DeadZone=0.000000,Sensitivity=1.000000,Exponent=1.000000,bInvert=False))
++AxisConfig=(AxisKeyName="OculusTouch_Right_FaceButton1",AxisProperties=(DeadZone=0.000000,Sensitivity=1.000000,Exponent=1.000000,bInvert=False))
++AxisConfig=(AxisKeyName="OculusTouch_Right_Trigger",AxisProperties=(DeadZone=0.000000,Sensitivity=1.000000,Exponent=1.000000,bInvert=False))
++AxisConfig=(AxisKeyName="OculusTouch_Right_FaceButton2",AxisProperties=(DeadZone=0.000000,Sensitivity=1.000000,Exponent=1.000000,bInvert=False))
++AxisConfig=(AxisKeyName="OculusTouch_Right_IndexPointing",AxisProperties=(DeadZone=0.000000,Sensitivity=1.000000,Exponent=1.000000,bInvert=False))
++AxisConfig=(AxisKeyName="OculusTouch_Right_ThumbUp",AxisProperties=(DeadZone=0.000000,Sensitivity=1.000000,Exponent=1.000000,bInvert=False))
++AxisConfig=(AxisKeyName="OculusTouchpad_Touchpad_X",AxisProperties=(DeadZone=0.000000,Sensitivity=1.000000,Exponent=1.000000,bInvert=False))
++AxisConfig=(AxisKeyName="OculusTouchpad_Touchpad_Y",AxisProperties=(DeadZone=0.000000,Sensitivity=1.000000,Exponent=1.000000,bInvert=False))
++AxisConfig=(AxisKeyName="SteamVR_Knuckles_Left_HandGrip",AxisProperties=(DeadZone=0.000000,Sensitivity=1.000000,Exponent=1.000000,bInvert=False))
++AxisConfig=(AxisKeyName="SteamVR_Knuckles_Left_IndexGrip",AxisProperties=(DeadZone=0.000000,Sensitivity=1.000000,Exponent=1.000000,bInvert=False))
++AxisConfig=(AxisKeyName="SteamVR_Knuckles_Left_MiddleGrip",AxisProperties=(DeadZone=0.000000,Sensitivity=1.000000,Exponent=1.000000,bInvert=False))
++AxisConfig=(AxisKeyName="SteamVR_Knuckles_Left_RingGrip",AxisProperties=(DeadZone=0.000000,Sensitivity=1.000000,Exponent=1.000000,bInvert=False))
++AxisConfig=(AxisKeyName="SteamVR_Knuckles_Left_PinkyGrip",AxisProperties=(DeadZone=0.000000,Sensitivity=1.000000,Exponent=1.000000,bInvert=False))
++AxisConfig=(AxisKeyName="SteamVR_Knuckles_Right_HandGrip",AxisProperties=(DeadZone=0.000000,Sensitivity=1.000000,Exponent=1.000000,bInvert=False))
++AxisConfig=(AxisKeyName="SteamVR_Knuckles_Right_IndexGrip",AxisProperties=(DeadZone=0.000000,Sensitivity=1.000000,Exponent=1.000000,bInvert=False))
++AxisConfig=(AxisKeyName="SteamVR_Knuckles_Right_MiddleGrip",AxisProperties=(DeadZone=0.000000,Sensitivity=1.000000,Exponent=1.000000,bInvert=False))
++AxisConfig=(AxisKeyName="SteamVR_Knuckles_Right_RingGrip",AxisProperties=(DeadZone=0.000000,Sensitivity=1.000000,Exponent=1.000000,bInvert=False))
++AxisConfig=(AxisKeyName="SteamVR_Knuckles_Right_PinkyGrip",AxisProperties=(DeadZone=0.000000,Sensitivity=1.000000,Exponent=1.000000,bInvert=False))
++AxisConfig=(AxisKeyName="Gamepad_LeftX",AxisProperties=(DeadZone=0.250000,Sensitivity=1.000000,Exponent=1.000000,bInvert=False))
++AxisConfig=(AxisKeyName="Gamepad_LeftY",AxisProperties=(DeadZone=0.250000,Sensitivity=1.000000,Exponent=1.000000,bInvert=False))
++AxisConfig=(AxisKeyName="Gamepad_RightX",AxisProperties=(DeadZone=0.250000,Sensitivity=1.000000,Exponent=1.000000,bInvert=False))
++AxisConfig=(AxisKeyName="Gamepad_RightY",AxisProperties=(DeadZone=0.250000,Sensitivity=1.000000,Exponent=1.000000,bInvert=False))
++AxisConfig=(AxisKeyName="MouseX",AxisProperties=(DeadZone=0.000000,Sensitivity=0.070000,Exponent=1.000000,bInvert=False))
++AxisConfig=(AxisKeyName="MouseY",AxisProperties=(DeadZone=0.000000,Sensitivity=0.070000,Exponent=1.000000,bInvert=False))
++AxisConfig=(AxisKeyName="MouseWheelAxis",AxisProperties=(DeadZone=0.000000,Sensitivity=1.000000,Exponent=1.000000,bInvert=False))
++AxisConfig=(AxisKeyName="Gamepad_LeftTriggerAxis",AxisProperties=(DeadZone=0.000000,Sensitivity=1.000000,Exponent=1.000000,bInvert=False))
++AxisConfig=(AxisKeyName="Gamepad_RightTriggerAxis",AxisProperties=(DeadZone=0.000000,Sensitivity=1.000000,Exponent=1.000000,bInvert=False))
++AxisConfig=(AxisKeyName="MotionController_Left_Thumbstick_X",AxisProperties=(DeadZone=0.000000,Sensitivity=1.000000,Exponent=1.000000,bInvert=False))
++AxisConfig=(AxisKeyName="MotionController_Left_Thumbstick_Y",AxisProperties=(DeadZone=0.000000,Sensitivity=1.000000,Exponent=1.000000,bInvert=False))
++AxisConfig=(AxisKeyName="MotionController_Left_TriggerAxis",AxisProperties=(DeadZone=0.000000,Sensitivity=1.000000,Exponent=1.000000,bInvert=False))
++AxisConfig=(AxisKeyName="MotionController_Left_Grip1Axis",AxisProperties=(DeadZone=0.000000,Sensitivity=1.000000,Exponent=1.000000,bInvert=False))
++AxisConfig=(AxisKeyName="MotionController_Left_Grip2Axis",AxisProperties=(DeadZone=0.000000,Sensitivity=1.000000,Exponent=1.000000,bInvert=False))
++AxisConfig=(AxisKeyName="MotionController_Right_Thumbstick_X",AxisProperties=(DeadZone=0.000000,Sensitivity=1.000000,Exponent=1.000000,bInvert=False))
++AxisConfig=(AxisKeyName="MotionController_Right_Thumbstick_Y",AxisProperties=(DeadZone=0.000000,Sensitivity=1.000000,Exponent=1.000000,bInvert=False))
++AxisConfig=(AxisKeyName="MotionController_Right_TriggerAxis",AxisProperties=(DeadZone=0.000000,Sensitivity=1.000000,Exponent=1.000000,bInvert=False))
++AxisConfig=(AxisKeyName="MotionController_Right_Grip1Axis",AxisProperties=(DeadZone=0.000000,Sensitivity=1.000000,Exponent=1.000000,bInvert=False))
++AxisConfig=(AxisKeyName="MotionController_Right_Grip2Axis",AxisProperties=(DeadZone=0.000000,Sensitivity=1.000000,Exponent=1.000000,bInvert=False))
++AxisConfig=(AxisKeyName="Gamepad_Special_Left_X",AxisProperties=(DeadZone=0.000000,Sensitivity=1.000000,Exponent=1.000000,bInvert=False))
++AxisConfig=(AxisKeyName="Gamepad_Special_Left_Y",AxisProperties=(DeadZone=0.000000,Sensitivity=1.000000,Exponent=1.000000,bInvert=False))
++AxisConfig=(AxisKeyName="MotionController_Left_Thumbstick_Z",AxisProperties=(DeadZone=0.000000,Sensitivity=1.000000,Exponent=1.000000,bInvert=False))
++AxisConfig=(AxisKeyName="MotionController_Right_Thumbstick_Z",AxisProperties=(DeadZone=0.000000,Sensitivity=1.000000,Exponent=1.000000,bInvert=False))
++AxisConfig=(AxisKeyName="OculusTouch_Left_Thumbstick",AxisProperties=(DeadZone=0.000000,Sensitivity=1.000000,Exponent=1.000000,bInvert=False))
++AxisConfig=(AxisKeyName="OculusTouch_Left_FaceButton1",AxisProperties=(DeadZone=0.000000,Sensitivity=1.000000,Exponent=1.000000,bInvert=False))
++AxisConfig=(AxisKeyName="OculusTouch_Left_Trigger",AxisProperties=(DeadZone=0.000000,Sensitivity=1.000000,Exponent=1.000000,bInvert=False))
++AxisConfig=(AxisKeyName="OculusTouch_Left_FaceButton2",AxisProperties=(DeadZone=0.000000,Sensitivity=1.000000,Exponent=1.000000,bInvert=False))
++AxisConfig=(AxisKeyName="OculusTouch_Left_IndexPointing",AxisProperties=(DeadZone=0.000000,Sensitivity=1.000000,Exponent=1.000000,bInvert=False))
++AxisConfig=(AxisKeyName="OculusTouch_Left_ThumbUp",AxisProperties=(DeadZone=0.000000,Sensitivity=1.000000,Exponent=1.000000,bInvert=False))
+bAltEnterTogglesFullscreen=True
+bF11TogglesFullscreen=True
+bUseMouseForTouch=False
+bEnableMouseSmoothing=True
+bEnableFOVScaling=True
+bCaptureMouseOnLaunch=True
+bDefaultViewportMouseLock=False
+bAlwaysShowTouchInterface=False
+bShowConsoleOnFourFingerTap=True
+bEnableGestureRecognizer=True
+bUseAutocorrect=False
+DefaultViewportMouseCaptureMode=CapturePermanently_IncludingInitialMouseDown
+DefaultViewportMouseLockMode=LockOnCapture
+FOVScale=0.011110
+DoubleClickTime=0.200000
++ActionMappings=(ActionName="MotionControllerL",bShift=False,bCtrl=False,bAlt=False,bCmd=False,Key=MotionController_Left_Trigger)
++ActionMappings=(ActionName="MotionControllerR",bShift=False,bCtrl=False,bAlt=False,bCmd=False,Key=MotionController_Right_Trigger)
++ActionMappings=(ActionName="PushAndPull",bShift=False,bCtrl=False,bAlt=False,bCmd=False,Key=MouseWheelAxis)
++ActionMappings=(ActionName="Quit",bShift=False,bCtrl=False,bAlt=False,bCmd=False,Key=Escape)
++ActionMappings=(ActionName="FlyMode",bShift=False,bCtrl=True,bAlt=False,bCmd=False,Key=One)
++ActionMappings=(ActionName="WalkMode",bShift=False,bCtrl=True,bAlt=False,bCmd=False,Key=Two)
++ActionMappings=(ActionName="OrbitMode",bShift=False,bCtrl=True,bAlt=False,bCmd=False,Key=Three)
++ActionMappings=(ActionName="VRMode",bShift=False,bCtrl=True,bAlt=False,bCmd=False,Key=Four)
++ActionMappings=(ActionName="ResetLocationL",bShift=False,bCtrl=False,bAlt=False,bCmd=False,Key=MotionController_Left_Grip1)
++ActionMappings=(ActionName="Bookmark_1",bShift=False,bCtrl=False,bAlt=False,bCmd=False,Key=One)
++ActionMappings=(ActionName="Bookmark_2",bShift=False,bCtrl=False,bAlt=False,bCmd=False,Key=Two)
++ActionMappings=(ActionName="Bookmark_3",bShift=False,bCtrl=False,bAlt=False,bCmd=False,Key=Three)
++ActionMappings=(ActionName="Bookmark_4",bShift=False,bCtrl=False,bAlt=False,bCmd=False,Key=Four)
++ActionMappings=(ActionName="Bookmark_5",bShift=False,bCtrl=False,bAlt=False,bCmd=False,Key=Five)
++ActionMappings=(ActionName="Bookmark_6",bShift=False,bCtrl=False,bAlt=False,bCmd=False,Key=Six)
++ActionMappings=(ActionName="Bookmark_7",bShift=False,bCtrl=False,bAlt=False,bCmd=False,Key=Seven)
++ActionMappings=(ActionName="Bookmark_8",bShift=False,bCtrl=False,bAlt=False,bCmd=False,Key=Eight)
++ActionMappings=(ActionName="Bookmark_9",bShift=False,bCtrl=False,bAlt=False,bCmd=False,Key=Nine)
++ActionMappings=(ActionName="OpenContextualMenu",bShift=False,bCtrl=False,bAlt=False,bCmd=False,Key=SpaceBar)
++ActionMappings=(ActionName="OrbitAroundObject",bShift=False,bCtrl=False,bAlt=False,bCmd=False,Key=RightMouseButton)
++ActionMappings=(ActionName="OrbitZoomOut",bShift=False,bCtrl=False,bAlt=False,bCmd=False,Key=MouseScrollDown)
++ActionMappings=(ActionName="OrbitZoomIn",bShift=False,bCtrl=False,bAlt=False,bCmd=False,Key=MouseScrollUp)
++ActionMappings=(ActionName="ShootLaser",bShift=False,bCtrl=False,bAlt=False,bCmd=False,Key=LeftMouseButton)
++ActionMappings=(ActionName="RotateView",bShift=False,bCtrl=False,bAlt=False,bCmd=False,Key=RightMouseButton)
++ActionMappings=(ActionName="SetPointOfinterest",bShift=False,bCtrl=False,bAlt=False,bCmd=False,Key=MiddleMouseButton)
++ActionMappings=(ActionName="LockToAxis",bShift=False,bCtrl=False,bAlt=False,bCmd=False,Key=LeftShift)
++AxisMappings=(AxisName="WalkForward",Scale=1.000000,Key=W)
++AxisMappings=(AxisName="WalkRight",Scale=1.000000,Key=D)
++AxisMappings=(AxisName="WalkForward",Scale=-1.000000,Key=S)
++AxisMappings=(AxisName="WalkRight",Scale=-1.000000,Key=A)
++AxisMappings=(AxisName="LookUp",Scale=1.000000,Key=MouseY)
++AxisMappings=(AxisName="LookRight",Scale=1.000000,Key=MouseX)
++AxisMappings=(AxisName="FlyUp",Scale=1.000000,Key=E)
++AxisMappings=(AxisName="FlyUp",Scale=-1.000000,Key=Q)
++AxisMappings=(AxisName="MotionControllerThumbLeft_Y",Scale=-1.000000,Key=MotionController_Left_Thumbstick_Y)
++AxisMappings=(AxisName="MotionControllerThumbRight_Y",Scale=-1.000000,Key=MotionController_Right_Thumbstick_Y)
++AxisMappings=(AxisName="MotionControllerThumbLeft_X",Scale=1.000000,Key=MotionController_Left_Thumbstick_X)
++AxisMappings=(AxisName="MotionControllerThumbRight_X",Scale=1.000000,Key=MotionController_Right_Thumbstick_X)
++AxisMappings=(AxisName="WalkForward",Scale=1.000000,Key=Gamepad_LeftY)
++AxisMappings=(AxisName="WalkRight",Scale=1.000000,Key=Gamepad_LeftX)
++AxisMappings=(AxisName="LookUp",Scale=-1.000000,Key=Gamepad_RightY)
++AxisMappings=(AxisName="LookRight",Scale=1.000000,Key=Gamepad_RightX)
++AxisMappings=(AxisName="pinch",Scale=1.000000,Key=Gesture_Pinch)
++AxisMappings=(AxisName="WalkForward",Scale=1.000000,Key=Up)
++AxisMappings=(AxisName="WalkForward",Scale=-1.000000,Key=Down)
++AxisMappings=(AxisName="WalkRight",Scale=1.000000,Key=Right)
++AxisMappings=(AxisName="WalkRight",Scale=-1.000000,Key=Left)
+DefaultTouchInterface=/Game/CollaborativeViewer/Blueprints/Pawn/Touch/TouchInterface.TouchInterface
+ConsoleKey=None
++ConsoleKeys=Tilde
+
+[Redirects]
+GameName=TP_CollaborativeBP
+
+[AdditionalFilesToAdd]
++Files=Enterprise/Templates/TP_CollaborativeBP/Content/*.*
diff --git a/Templates/TP_CollaborativeBP/contents.txt b/Templates/TP_CollaborativeBP/contents.txt
--- a/Templates/TP_CollaborativeBP/contents.txt
+++ b/Templates/TP_CollaborativeBP/contents.txt
@@ -1,4 +1,4 @@
-../../../Templates/TP_Collaborative_V/manifest.json
-../../../Templates/TP_Collaborative_V/Config/config.ini
-../../../Templates/TP_Collaborative_V/Media/TP_Collaborative.png
-../../../Templates/TP_Collaborative_V/Media/TP_Collaborative_V_Preview.png
+../../../Templates/TP_CollaborativeBP/manifest.json
+../../../Templates/TP_CollaborativeBP/Config/config.ini
+../../../Templates/TP_CollaborativeBP/Media/TP_CollaborativeBP.png
+../../../Templates/TP_CollaborativeBP/Media/TP_CollaborativeBP_Preview.png
diff --git a/Templates/TP_CollaborativeBP/manifest.json b/Templates/TP_CollaborativeBP/manifest.json
--- a/Templates/TP_CollaborativeBP/manifest.json
+++ b/Templates/TP_CollaborativeBP/manifest.json
@@ -19,12 +19,12 @@
 	"ClassTypes": "",
 	"Category": "",
 	"FocusAsset": "",
-	"Thumbnail": "TP_Collaborative.png",
+	"Thumbnail": "TP_CollaborativeBP.png",
 	"Screenshots":
 	[
-		"TP_Collaborative_V_Preview.png"
+		"TP_CollaborativeBP_Preview.png"
 	],
-	"SortKey": "TP_Collaborative_V",
+	"SortKey": "TP_CollaborativeBP",
 	"AdditionalFeaturePacks":
 	[
 	]
