UE-82445 Collab Viewer - non-VR player spawn in the floor and can't access menu when joining a VR host
UE-82529-Collaborative viewer - Menu navigation is not working with Vive hardware
UE-82428-Collaborative template shouldn't not enable the SunPosition plugin
UE-81843 - Colab Viewer: VR motion controllers don't work as intended as a VR Client
UE-82437 - CollabViewer - Oculus - Possible to Multi Select menu options
UE-82431-CollabViewser - Oculus - UI Buttons are still slightly visible in VR View

diff --git a/Templates/TP_CollaborativeBP/Config/DefaultEngine.ini b/Templates/TP_CollaborativeBP/Config/DefaultEngine.ini
--- a/Templates/TP_CollaborativeBP/Config/DefaultEngine.ini
+++ b/Templates/TP_CollaborativeBP/Config/DefaultEngine.ini
@@ -1,55 +1,18 @@
+[/Script/Engine.Engine]
++ActiveGameNameRedirects=(OldGameName="TP_BlankBP",NewGameName="/Script/MyProject3")
++ActiveGameNameRedirects=(OldGameName="/Script/TP_BlankBP",NewGameName="/Script/MyProject3")
+
 [/Script/HardwareTargeting.HardwareTargetingSettings]
-; quality settings
 TargetedHardwareClass=Desktop
 AppliedTargetedHardwareClass=Desktop
 DefaultGraphicsPerformance=Maximum
 AppliedDefaultGraphicsPerformance=Maximum
 
-
-[/Script/Engine.RendererSettings]
-; exposure control and tone mapping settings
-r.Tonemapper.Sharpen=2
-r.TemporalAA.Upsampling=True
-r.DefaultFeature.AutoExposure.ExtendDefaultLuminanceRange=True
-
-;dynamic gi settings
-r.LightPropagationVolume=0
-
-; quality settings
-r.ReflectionCaptureResolution=2048
-r.AllowStaticLighting=True
-r.HighResScreenshotDelay=8
-r.DefaultBackBufferPixelFormat=4
-r.AllowGlobalClipPlane=False
-
-; shader settings
-r.SupportSkyAtmosphereAffectsHeightFog=True
-r.DefaultFeature.LensFlare=True
-r.ClearCoatNormal=False
-r.NormalMapsForStaticLighting=False
-r.CustomDepth=3
-
 [/Script/EngineSettings.GameMapsSettings]
-GameInstanceClass=/Game/CollaborativeViewer/Blueprints/GameInstance/BP_CollaborativeViewer_GameInstance.BP_CollaborativeViewer_GameInstance_C
-GlobalDefaultGameMode=/Game/CollaborativeViewer/Blueprints/GameMode/BP_CollaborativeViewer_GameMode.BP_CollaborativeViewer_GameMode_C
 EditorStartupMap=/Game/CollaborativeViewer/Maps/CollaborativeViewer_P.CollaborativeViewer_P
 GameDefaultMap=/Game/CollaborativeViewer/Maps/MainMenu.MainMenu
-
-
-[/Script/Engine.UserInterfaceSettings]
-UIScaleRule=ShortestSide
-UIScaleCurve=(EditorCurveData=(Keys=((Time=480.000000,Value=0.444000),(Time=720.000000,Value=0.666000),(Time=1080.000000,Value=1.000000),(Time=8640.000000,Value=8.000000)),PreInfinityExtrap=RCCE_Constant,PostInfinityExtrap=RCCE_Constant,DefaultValue=340282346638528859811704183484516925440.000000),ExternalCurve=None)
-bAllowHighDPIInGameMode=False
-ApplicationScale=1.000000
-
-[/Script/IOSRuntimeSettings.IOSRuntimeSettings]
-RemoteServerName=10.9.101.128
-RSyncUsername=compile
-MobileProvision=
-SigningCertificate=
-
-[/Script/Engine.Engine]
-NearClipPlane=1.000000
+GlobalDefaultGameMode=/Game/CollaborativeViewer/Blueprints/GameMode/BP_CollaborativeViewer_GameMode.BP_CollaborativeViewer_GameMode_C
+GameInstanceClass=/Game/CollaborativeViewer/Blueprints/GameInstance/BP_CollaborativeViewer_GameInstance.BP_CollaborativeViewer_GameInstance_C
 
 [/Script/Engine.CollisionProfile]
 +Profiles=(Name="NoCollision",CollisionEnabled=NoCollision,ObjectTypeName="WorldStatic",CustomResponses=((Channel="Visibility",Response=ECR_Ignore),(Channel="Camera",Response=ECR_Ignore)),HelpMessage="No collision",bCanModify=False)
@@ -84,9 +47,53 @@ NearClipPlane=1.000000
 +CollisionChannelRedirects=(OldName="VehicleMovement",NewName="Vehicle")
 +CollisionChannelRedirects=(OldName="PawnMovement",NewName="Pawn")
 
-[/Script/UnrealEd.UnrealEdEngine]
-; Remove Engine Template maps
--TemplateMapInfos=(ThumbnailTexture=Texture2D'/Engine/Maps/Templates/Thumbnails/Default.Default',Map="/Engine/Maps/Templates/Template_Default")
--TemplateMapInfos=(ThumbnailTexture=Texture2D'/Engine/Maps/Templates/Thumbnails/TimeOfDay.TimeOfDay',Map="/Engine/Maps/Templates/TimeOfDay_Default")
--TemplateMapInfos=(ThumbnailTexture=Texture2D'/Engine/Maps/Templates/Thumbnails/VR-Basic.VR-Basic',Map="/Engine/Maps/Templates/VR-Basic")
-+TemplateMapInfos=(ThumbnailTexture=Texture2D'/Engine/Maps/Templates/Thumbnails/Basic.Basic',Map="/Engine/Maps/Templates/Basic")
+[/Script/Engine.RendererSettings]
+r.CustomDepth=3
+
+[/Script/Engine.PhysicsSettings]
+DefaultGravityZ=-980.000000
+DefaultTerminalVelocity=4000.000000
+DefaultFluidFriction=0.300000
+SimulateScratchMemorySize=262144
+RagdollAggregateThreshold=4
+TriangleMeshTriangleMinAreaThreshold=5.000000
+bEnableShapeSharing=False
+bEnablePCM=True
+bEnableStabilization=False
+bWarnMissingLocks=True
+bEnable2DPhysics=False
+PhysicErrorCorrection=(PingExtrapolation=0.100000,PingLimit=100.000000,ErrorPerLinearDifference=1.000000,ErrorPerAngularDifference=1.000000,MaxRestoredStateError=1.000000,MaxLinearHardSnapDistance=400.000000,PositionLerp=0.000000,AngleLerp=0.400000,LinearVelocityCoefficient=100.000000,AngularVelocityCoefficient=10.000000,ErrorAccumulationSeconds=0.500000,ErrorAccumulationDistanceSq=15.000000,ErrorAccumulationSimilarity=100.000000)
+LockedAxis=Invalid
+DefaultDegreesOfFreedom=Full3D
+BounceThresholdVelocity=200.000000
+FrictionCombineMode=Average
+RestitutionCombineMode=Average
+MaxAngularVelocity=3600.000000
+MaxDepenetrationVelocity=0.000000
+ContactOffsetMultiplier=0.020000
+MinContactOffset=2.000000
+MaxContactOffset=8.000000
+bSimulateSkeletalMeshOnDedicatedServer=True
+DefaultShapeComplexity=CTF_UseSimpleAndComplex
+bDefaultHasComplexCollision=True
+bSuppressFaceRemapTable=False
+bSupportUVFromHitResults=False
+bDisableActiveActors=False
+bDisableKinematicStaticPairs=False
+bDisableKinematicKinematicPairs=False
+bDisableCCD=False
+bEnableEnhancedDeterminism=False
+AnimPhysicsMinDeltaTime=0.000000
+bSimulateAnimPhysicsAfterReset=False
+MaxPhysicsDeltaTime=0.033333
+bSubstepping=False
+bSubsteppingAsync=False
+MaxSubstepDeltaTime=0.016667
+MaxSubsteps=6
+SyncSceneSmoothingFactor=0.000000
+InitialAverageFrameRate=0.016667
+PhysXTreeRebuildRate=10
+DefaultBroadphaseSettings=(bUseMBPOnClient=False,bUseMBPOnServer=False,bUseMBPOuterBounds=False,MBPBounds=(Min=(X=0.000000,Y=0.000000,Z=0.000000),Max=(X=0.000000,Y=0.000000,Z=0.000000),IsValid=0),MBPOuterBounds=(Min=(X=0.000000,Y=0.000000,Z=0.000000),Max=(X=0.000000,Y=0.000000,Z=0.000000),IsValid=0),MBPNumSubdivs=2)
+ChaosSettings=(DefaultThreadingModel=TaskGraph,DedicatedThreadTickMode=VariableCappedWithTarget,DedicatedThreadBufferMode=Double)
+
+
diff --git a/Templates/TP_CollaborativeBP/Config/DefaultGame.ini b/Templates/TP_CollaborativeBP/Config/DefaultGame.ini
--- a/Templates/TP_CollaborativeBP/Config/DefaultGame.ini
+++ b/Templates/TP_CollaborativeBP/Config/DefaultGame.ini
@@ -1,85 +1,7 @@
 
-[/Script/UnrealEd.ProjectPackagingSettings]
-Build=IfProjectHasCode
-BuildConfiguration=PPBC_Development
-StagingDirectory=(Path="")
-FullRebuild=False
-ForDistribution=False
-IncludeDebugFiles=False
-BlueprintNativizationMethod=Disabled
-bIncludeNativizedAssetsInProjectGeneration=False
-bExcludeMonolithicEngineHeadersInNativizedCode=False
-UsePakFile=True
-bGenerateChunks=False
-bGenerateNoChunks=False
-bChunkHardReferencesOnly=False
-bForceOneChunkPerFile=False
-MaxChunkSize=0
-bBuildHttpChunkInstallData=False
-HttpChunkInstallDataDirectory=(Path="")
-PakFileCompressionFormats=
-PakFileAdditionalCompressionOptions=
-HttpChunkInstallDataVersion=
-IncludePrerequisites=True
-IncludeAppLocalPrerequisites=False
-bShareMaterialShaderCode=False
-bSharedMaterialNativeLibraries=False
-ApplocalPrerequisitesDirectory=(Path="")
-IncludeCrashReporter=False
-InternationalizationPreset=English
--CulturesToStage=en
-+CulturesToStage=en
-bCookAll=False
-bCookMapsOnly=False
-bCompressed=False
-bEncryptIniFiles=False
-bEncryptPakIndex=False
-GenerateEarlyDownloaderPakFile=False
-bSkipEditorContent=False
-bSkipMovies=False
--EarlyDownloaderPakFileFiles=...\Content\Internationalization\...\*.icu
--EarlyDownloaderPakFileFiles=...\Content\Internationalization\...\*.brk
--EarlyDownloaderPakFileFiles=...\Content\Internationalization\...\*.res
--EarlyDownloaderPakFileFiles=...\Content\Internationalization\...\*.nrm
--EarlyDownloaderPakFileFiles=...\Content\Internationalization\...\*.cfu
--EarlyDownloaderPakFileFiles=...\Content\Localization\...\*.*
--EarlyDownloaderPakFileFiles=...\Content\Localization\*.*
--EarlyDownloaderPakFileFiles=...\Content\Certificates\...\*.*
--EarlyDownloaderPakFileFiles=...\Content\Certificates\*.*
--EarlyDownloaderPakFileFiles=-...\Content\Localization\Game\...\*.*
--EarlyDownloaderPakFileFiles=-...\Content\Localization\Game\*.*
--EarlyDownloaderPakFileFiles=...\Config\...\*.ini
--EarlyDownloaderPakFileFiles=...\Config\*.ini
--EarlyDownloaderPakFileFiles=...\Engine\GlobalShaderCache*.bin
--EarlyDownloaderPakFileFiles=...\Content\ShaderArchive-Global*.ushaderbytecode
--EarlyDownloaderPakFileFiles=...\Content\Slate\*.*
--EarlyDownloaderPakFileFiles=...\Content\Slate\...\*.*
--EarlyDownloaderPakFileFiles=...\*.upluginmanifest
--EarlyDownloaderPakFileFiles=...\*.uproject
--EarlyDownloaderPakFileFiles=...\global_sf*.metalmap
-+EarlyDownloaderPakFileFiles=...\Content\Internationalization\...\*.icu
-+EarlyDownloaderPakFileFiles=...\Content\Internationalization\...\*.brk
-+EarlyDownloaderPakFileFiles=...\Content\Internationalization\...\*.res
-+EarlyDownloaderPakFileFiles=...\Content\Internationalization\...\*.nrm
-+EarlyDownloaderPakFileFiles=...\Content\Internationalization\...\*.cfu
-+EarlyDownloaderPakFileFiles=...\Content\Localization\...\*.*
-+EarlyDownloaderPakFileFiles=...\Content\Localization\*.*
-+EarlyDownloaderPakFileFiles=...\Content\Certificates\...\*.*
-+EarlyDownloaderPakFileFiles=...\Content\Certificates\*.*
-+EarlyDownloaderPakFileFiles=-...\Content\Localization\Game\...\*.*
-+EarlyDownloaderPakFileFiles=-...\Content\Localization\Game\*.*
-+EarlyDownloaderPakFileFiles=...\Config\...\*.ini
-+EarlyDownloaderPakFileFiles=...\Config\*.ini
-+EarlyDownloaderPakFileFiles=...\Engine\GlobalShaderCache*.bin
-+EarlyDownloaderPakFileFiles=...\Content\ShaderArchive-Global*.ushaderbytecode
-+EarlyDownloaderPakFileFiles=...\Content\Slate\*.*
-+EarlyDownloaderPakFileFiles=...\Content\Slate\...\*.*
-+EarlyDownloaderPakFileFiles=...\*.upluginmanifest
-+EarlyDownloaderPakFileFiles=...\*.uproject
-+EarlyDownloaderPakFileFiles=...\global_sf*.metalmap
-+MapsToCook=(FilePath="/Game/CollaborativeViewer/Maps/CollaborativeViewer_P")
-+MapsToCook=(FilePath="/Game/CollaborativeViewer/Maps/MainMenu")
-bNativizeBlueprintAssets=False
-bNativizeOnlySelectedBlueprints=False
 
 
+[/Script/EngineSettings.GeneralProjectSettings]
+ProjectID=C684876B46D5306ADA6286AA20FA1808
+bStartInVR=True
+
diff --git a/Templates/TP_CollaborativeBP/Config/IOS/IOSEngine.ini b/Templates/TP_CollaborativeBP/Config/IOS/IOSEngine.ini
new file mode 100644
--- /dev/null
+++ b/Templates/TP_CollaborativeBP/Config/IOS/IOSEngine.ini
@@ -0,0 +1,8 @@
+[/Script/IOSRuntimeSettings.IOSRuntimeSettings]
+;AudioSampleRate=48000
+AudioMaxChannels=16
+;AudioCallbackBufferFrameSize=1024
+;AudioNumBuffersToEnqueue=2
+;AudioNumSourceWorkers=0
+[OnlineSubsystem]
+DefaultPlatformService=Null
\ No newline at end of file
diff --git a/Templates/TP_CollaborativeBP/TP_CollaborativeBP.uproject b/Templates/TP_CollaborativeBP/TP_CollaborativeBP.uproject
--- a/Templates/TP_CollaborativeBP/TP_CollaborativeBP.uproject
+++ b/Templates/TP_CollaborativeBP/TP_CollaborativeBP.uproject
@@ -4,10 +4,6 @@
 	"Category": "",
 	"Description": "",
 	"Plugins": [
-		{
-			"Name": "SunPosition",
-			"Enabled": true
-		},
 		{
 			"Name": "DataPrepEditor",
 			"Enabled": true
