UE-82439 Collab Viewer - third player cannot access the menu
UE-82868 CollabViewer - Oculus - Stuck in drawing mode when both Index triggers are pushed.
UE-82872 Collaborative Viewer - Can't type in user name on start screen on mobile
UE-82431 CollabViewer - Oculus - UI Buttons are still slightly visible in VR View
UE-82434 CollabViewer - Oculus - VR Presence is slightly tilted after teleported at Bookmark

diff --git a/Templates/TP_CollaborativeBP/Config/DefaultGame.ini b/Templates/TP_CollaborativeBP/Config/DefaultGame.ini
--- a/Templates/TP_CollaborativeBP/Config/DefaultGame.ini
+++ b/Templates/TP_CollaborativeBP/Config/DefaultGame.ini
@@ -3,5 +3,5 @@
 
 [/Script/EngineSettings.GeneralProjectSettings]
 ProjectID=C684876B46D5306ADA6286AA20FA1808
-bStartInVR=True
+bStartInVR=False
 
