UE-82851 - VR Clients sees the context menu displayed twice
UE-83039 - InputAxisKey Event specify warnings for non bindable blueprints for MotionControllers referring to Compiler BP_VRPawn & Compiler BP_VRPawn

diff --git a/Templates/TP_CollaborativeBP/Config/DefaultInput.ini b/Templates/TP_CollaborativeBP/Config/DefaultInput.ini
--- a/Templates/TP_CollaborativeBP/Config/DefaultInput.ini
+++ b/Templates/TP_CollaborativeBP/Config/DefaultInput.ini
@@ -1,6 +1,42 @@
 
 [/Script/Engine.InputSettings]
-+AxisConfig=(AxisKeyName="OculusTouch_Right_Thumbstick",AxisProperties=(DeadZone=0.000000,Sensitivity=1.000000,Exponent=1.000000,bInvert=False))
++AxisConfig=(AxisKeyName="ValveIndex_Left_Grip_Force",AxisProperties=(DeadZone=0.000000,Sensitivity=1.000000,Exponent=1.000000,bInvert=False))
++AxisConfig=(AxisKeyName="ValveIndex_Left_Trigger_Axis",AxisProperties=(DeadZone=0.000000,Sensitivity=1.000000,Exponent=1.000000,bInvert=False))
++AxisConfig=(AxisKeyName="ValveIndex_Left_Thumbstick_X",AxisProperties=(DeadZone=0.000000,Sensitivity=1.000000,Exponent=1.000000,bInvert=False))
++AxisConfig=(AxisKeyName="ValveIndex_Left_Thumbstick_Y",AxisProperties=(DeadZone=0.000000,Sensitivity=1.000000,Exponent=1.000000,bInvert=False))
++AxisConfig=(AxisKeyName="ValveIndex_Left_Trackpad_X",AxisProperties=(DeadZone=0.000000,Sensitivity=1.000000,Exponent=1.000000,bInvert=False))
++AxisConfig=(AxisKeyName="ValveIndex_Left_Trackpad_Y",AxisProperties=(DeadZone=0.000000,Sensitivity=1.000000,Exponent=1.000000,bInvert=False))
++AxisConfig=(AxisKeyName="ValveIndex_Left_Trackpad_Force",AxisProperties=(DeadZone=0.000000,Sensitivity=1.000000,Exponent=1.000000,bInvert=False))
++AxisConfig=(AxisKeyName="ValveIndex_Left_Trackpad_Touch",AxisProperties=(DeadZone=0.000000,Sensitivity=1.000000,Exponent=1.000000,bInvert=False))
++AxisConfig=(AxisKeyName="ValveIndex_Right_Grip_Axis",AxisProperties=(DeadZone=0.000000,Sensitivity=1.000000,Exponent=1.000000,bInvert=False))
++AxisConfig=(AxisKeyName="ValveIndex_Right_Grip_Force",AxisProperties=(DeadZone=0.000000,Sensitivity=1.000000,Exponent=1.000000,bInvert=False))
++AxisConfig=(AxisKeyName="ValveIndex_Right_Trigger_Axis",AxisProperties=(DeadZone=0.000000,Sensitivity=1.000000,Exponent=1.000000,bInvert=False))
++AxisConfig=(AxisKeyName="ValveIndex_Right_Thumbstick_X",AxisProperties=(DeadZone=0.000000,Sensitivity=1.000000,Exponent=1.000000,bInvert=False))
++AxisConfig=(AxisKeyName="ValveIndex_Right_Thumbstick_Y",AxisProperties=(DeadZone=0.000000,Sensitivity=1.000000,Exponent=1.000000,bInvert=False))
++AxisConfig=(AxisKeyName="ValveIndex_Right_Trackpad_X",AxisProperties=(DeadZone=0.000000,Sensitivity=1.000000,Exponent=1.000000,bInvert=False))
++AxisConfig=(AxisKeyName="ValveIndex_Right_Trackpad_Y",AxisProperties=(DeadZone=0.000000,Sensitivity=1.000000,Exponent=1.000000,bInvert=False))
++AxisConfig=(AxisKeyName="ValveIndex_Right_Trackpad_Force",AxisProperties=(DeadZone=0.000000,Sensitivity=1.000000,Exponent=1.000000,bInvert=False))
++AxisConfig=(AxisKeyName="MagicLeap_Left_Trigger_Axis",AxisProperties=(DeadZone=0.000000,Sensitivity=1.000000,Exponent=1.000000,bInvert=False))
++AxisConfig=(AxisKeyName="MagicLeap_Left_Trackpad_X",AxisProperties=(DeadZone=0.000000,Sensitivity=1.000000,Exponent=1.000000,bInvert=False))
++AxisConfig=(AxisKeyName="MagicLeap_Left_Trackpad_Y",AxisProperties=(DeadZone=0.000000,Sensitivity=1.000000,Exponent=1.000000,bInvert=False))
++AxisConfig=(AxisKeyName="MagicLeap_Left_Trackpad_Force",AxisProperties=(DeadZone=0.000000,Sensitivity=1.000000,Exponent=1.000000,bInvert=False))
++AxisConfig=(AxisKeyName="MagicLeap_Left_Touch1_X",AxisProperties=(DeadZone=0.000000,Sensitivity=1.000000,Exponent=1.000000,bInvert=False))
++AxisConfig=(AxisKeyName="MagicLeap_Left_Touch1_Y",AxisProperties=(DeadZone=0.000000,Sensitivity=1.000000,Exponent=1.000000,bInvert=False))
++AxisConfig=(AxisKeyName="MagicLeap_Left_Touch1_Force",AxisProperties=(DeadZone=0.000000,Sensitivity=1.000000,Exponent=1.000000,bInvert=False))
++AxisConfig=(AxisKeyName="MagicLeap_Right_Trigger_Axis",AxisProperties=(DeadZone=0.000000,Sensitivity=1.000000,Exponent=1.000000,bInvert=False))
++AxisConfig=(AxisKeyName="MagicLeap_Right_Trackpad_X",AxisProperties=(DeadZone=0.000000,Sensitivity=1.000000,Exponent=1.000000,bInvert=False))
++AxisConfig=(AxisKeyName="MagicLeap_Right_Trackpad_Y",AxisProperties=(DeadZone=0.000000,Sensitivity=1.000000,Exponent=1.000000,bInvert=False))
++AxisConfig=(AxisKeyName="MagicLeap_Right_Trackpad_Force",AxisProperties=(DeadZone=0.000000,Sensitivity=1.000000,Exponent=1.000000,bInvert=False))
++AxisConfig=(AxisKeyName="MagicLeap_Right_Touch1_X",AxisProperties=(DeadZone=0.000000,Sensitivity=1.000000,Exponent=1.000000,bInvert=False))
++AxisConfig=(AxisKeyName="MagicLeap_Right_Touch1_Y",AxisProperties=(DeadZone=0.000000,Sensitivity=1.000000,Exponent=1.000000,bInvert=False))
++AxisConfig=(AxisKeyName="MagicLeap_Right_Touch1_Force",AxisProperties=(DeadZone=0.000000,Sensitivity=1.000000,Exponent=1.000000,bInvert=False))
++AxisConfig=(AxisKeyName="OculusTouch_Left_Thumbstick",AxisProperties=(DeadZone=0.000000,Sensitivity=1.000000,Exponent=1.000000,bInvert=False))
++AxisConfig=(AxisKeyName="OculusTouch_Left_FaceButton1",AxisProperties=(DeadZone=0.000000,Sensitivity=1.000000,Exponent=1.000000,bInvert=False))
++AxisConfig=(AxisKeyName="OculusTouch_Left_Trigger",AxisProperties=(DeadZone=0.000000,Sensitivity=1.000000,Exponent=1.000000,bInvert=False))
++AxisConfig=(AxisKeyName="OculusTouch_Left_FaceButton2",AxisProperties=(DeadZone=0.000000,Sensitivity=1.000000,Exponent=1.000000,bInvert=False))
++AxisConfig=(AxisKeyName="OculusTouch_Left_IndexPointing",AxisProperties=(DeadZone=0.000000,Sensitivity=1.000000,Exponent=1.000000,bInvert=False))
++AxisConfig=(AxisKeyName="OculusTouch_Left_ThumbUp",AxisProperties=(DeadZone=0.000000,Sensitivity=1.000000,Exponent=1.000000,bInvert=False))
++AxisConfig=(AxisKeyName="OculusTouch_Right_Thumbstick",AxisProperties=(DeadZone=0.000000,Sensitivity=100.000000,Exponent=1.000000,bInvert=False))
 +AxisConfig=(AxisKeyName="OculusTouch_Right_FaceButton1",AxisProperties=(DeadZone=0.000000,Sensitivity=1.000000,Exponent=1.000000,bInvert=False))
 +AxisConfig=(AxisKeyName="OculusTouch_Right_Trigger",AxisProperties=(DeadZone=0.000000,Sensitivity=1.000000,Exponent=1.000000,bInvert=False))
 +AxisConfig=(AxisKeyName="OculusTouch_Right_FaceButton2",AxisProperties=(DeadZone=0.000000,Sensitivity=1.000000,Exponent=1.000000,bInvert=False))
@@ -41,12 +77,39 @@
 +AxisConfig=(AxisKeyName="Gamepad_Special_Left_Y",AxisProperties=(DeadZone=0.000000,Sensitivity=1.000000,Exponent=1.000000,bInvert=False))
 +AxisConfig=(AxisKeyName="MotionController_Left_Thumbstick_Z",AxisProperties=(DeadZone=0.000000,Sensitivity=1.000000,Exponent=1.000000,bInvert=False))
 +AxisConfig=(AxisKeyName="MotionController_Right_Thumbstick_Z",AxisProperties=(DeadZone=0.000000,Sensitivity=1.000000,Exponent=1.000000,bInvert=False))
-+AxisConfig=(AxisKeyName="OculusTouch_Left_Thumbstick",AxisProperties=(DeadZone=0.000000,Sensitivity=1.000000,Exponent=1.000000,bInvert=False))
-+AxisConfig=(AxisKeyName="OculusTouch_Left_FaceButton1",AxisProperties=(DeadZone=0.000000,Sensitivity=1.000000,Exponent=1.000000,bInvert=False))
-+AxisConfig=(AxisKeyName="OculusTouch_Left_Trigger",AxisProperties=(DeadZone=0.000000,Sensitivity=1.000000,Exponent=1.000000,bInvert=False))
-+AxisConfig=(AxisKeyName="OculusTouch_Left_FaceButton2",AxisProperties=(DeadZone=0.000000,Sensitivity=1.000000,Exponent=1.000000,bInvert=False))
-+AxisConfig=(AxisKeyName="OculusTouch_Left_IndexPointing",AxisProperties=(DeadZone=0.000000,Sensitivity=1.000000,Exponent=1.000000,bInvert=False))
-+AxisConfig=(AxisKeyName="OculusTouch_Left_ThumbUp",AxisProperties=(DeadZone=0.000000,Sensitivity=1.000000,Exponent=1.000000,bInvert=False))
++AxisConfig=(AxisKeyName="Daydream_Left_Trackpad_X",AxisProperties=(DeadZone=0.000000,Sensitivity=1.000000,Exponent=1.000000,bInvert=False))
++AxisConfig=(AxisKeyName="Daydream_Left_Trackpad_Y",AxisProperties=(DeadZone=0.000000,Sensitivity=1.000000,Exponent=1.000000,bInvert=False))
++AxisConfig=(AxisKeyName="Daydream_Right_Trackpad_X",AxisProperties=(DeadZone=0.000000,Sensitivity=1.000000,Exponent=1.000000,bInvert=False))
++AxisConfig=(AxisKeyName="Daydream_Right_Trackpad_Y",AxisProperties=(DeadZone=0.000000,Sensitivity=1.000000,Exponent=1.000000,bInvert=False))
++AxisConfig=(AxisKeyName="Vive_Left_Trigger_Axis",AxisProperties=(DeadZone=0.000000,Sensitivity=1.000000,Exponent=1.000000,bInvert=False))
++AxisConfig=(AxisKeyName="Vive_Left_Trackpad_X",AxisProperties=(DeadZone=0.000000,Sensitivity=1.000000,Exponent=1.000000,bInvert=False))
++AxisConfig=(AxisKeyName="Vive_Left_Trackpad_Y",AxisProperties=(DeadZone=0.000000,Sensitivity=1.000000,Exponent=1.000000,bInvert=False))
++AxisConfig=(AxisKeyName="Vive_Right_Trigger_Axis",AxisProperties=(DeadZone=0.000000,Sensitivity=1.000000,Exponent=1.000000,bInvert=False))
++AxisConfig=(AxisKeyName="Vive_Right_Trackpad_X",AxisProperties=(DeadZone=0.000000,Sensitivity=1.000000,Exponent=1.000000,bInvert=False))
++AxisConfig=(AxisKeyName="Vive_Right_Trackpad_Y",AxisProperties=(DeadZone=0.000000,Sensitivity=1.000000,Exponent=1.000000,bInvert=False))
++AxisConfig=(AxisKeyName="MixedReality_Left_Trigger_Axis",AxisProperties=(DeadZone=0.000000,Sensitivity=1.000000,Exponent=1.000000,bInvert=False))
++AxisConfig=(AxisKeyName="MixedReality_Left_Thumbstick_X",AxisProperties=(DeadZone=0.000000,Sensitivity=1.000000,Exponent=1.000000,bInvert=False))
++AxisConfig=(AxisKeyName="MixedReality_Left_Thumbstick_Y",AxisProperties=(DeadZone=0.000000,Sensitivity=1.000000,Exponent=1.000000,bInvert=False))
++AxisConfig=(AxisKeyName="MixedReality_Left_Trackpad_X",AxisProperties=(DeadZone=0.000000,Sensitivity=1.000000,Exponent=1.000000,bInvert=False))
++AxisConfig=(AxisKeyName="MixedReality_Left_Trackpad_Y",AxisProperties=(DeadZone=0.000000,Sensitivity=1.000000,Exponent=1.000000,bInvert=False))
++AxisConfig=(AxisKeyName="MixedReality_Right_Trigger_Axis",AxisProperties=(DeadZone=0.000000,Sensitivity=1.000000,Exponent=1.000000,bInvert=False))
++AxisConfig=(AxisKeyName="MixedReality_Right_Thumbstick_X",AxisProperties=(DeadZone=0.000000,Sensitivity=1.000000,Exponent=1.000000,bInvert=False))
++AxisConfig=(AxisKeyName="MixedReality_Right_Thumbstick_Y",AxisProperties=(DeadZone=0.000000,Sensitivity=1.000000,Exponent=1.000000,bInvert=False))
++AxisConfig=(AxisKeyName="MixedReality_Right_Trackpad_X",AxisProperties=(DeadZone=0.000000,Sensitivity=1.000000,Exponent=1.000000,bInvert=False))
++AxisConfig=(AxisKeyName="MixedReality_Right_Trackpad_Y",AxisProperties=(DeadZone=0.000000,Sensitivity=1.000000,Exponent=1.000000,bInvert=False))
++AxisConfig=(AxisKeyName="OculusGo_Left_Trackpad_X",AxisProperties=(DeadZone=0.000000,Sensitivity=1.000000,Exponent=1.000000,bInvert=False))
++AxisConfig=(AxisKeyName="OculusGo_Left_Trackpad_Y",AxisProperties=(DeadZone=0.000000,Sensitivity=1.000000,Exponent=1.000000,bInvert=False))
++AxisConfig=(AxisKeyName="OculusGo_Right_Trackpad_X",AxisProperties=(DeadZone=0.000000,Sensitivity=1.000000,Exponent=1.000000,bInvert=False))
++AxisConfig=(AxisKeyName="OculusGo_Right_Trackpad_Y",AxisProperties=(DeadZone=0.000000,Sensitivity=1.000000,Exponent=1.000000,bInvert=False))
++AxisConfig=(AxisKeyName="OculusTouch_Left_Grip_Axis",AxisProperties=(DeadZone=0.000000,Sensitivity=1.000000,Exponent=1.000000,bInvert=False))
++AxisConfig=(AxisKeyName="OculusTouch_Left_Trigger_Axis",AxisProperties=(DeadZone=0.000000,Sensitivity=1.000000,Exponent=1.000000,bInvert=False))
++AxisConfig=(AxisKeyName="OculusTouch_Left_Thumbstick_X",AxisProperties=(DeadZone=0.000000,Sensitivity=1.000000,Exponent=1.000000,bInvert=False))
++AxisConfig=(AxisKeyName="OculusTouch_Left_Thumbstick_Y",AxisProperties=(DeadZone=0.000000,Sensitivity=1.000000,Exponent=1.000000,bInvert=False))
++AxisConfig=(AxisKeyName="OculusTouch_Right_Grip_Axis",AxisProperties=(DeadZone=0.000000,Sensitivity=1.000000,Exponent=1.000000,bInvert=False))
++AxisConfig=(AxisKeyName="OculusTouch_Right_Trigger_Axis",AxisProperties=(DeadZone=0.000000,Sensitivity=1.000000,Exponent=1.000000,bInvert=False))
++AxisConfig=(AxisKeyName="OculusTouch_Right_Thumbstick_X",AxisProperties=(DeadZone=0.000000,Sensitivity=1.000000,Exponent=1.000000,bInvert=False))
++AxisConfig=(AxisKeyName="OculusTouch_Right_Thumbstick_Y",AxisProperties=(DeadZone=0.000000,Sensitivity=1.000000,Exponent=1.000000,bInvert=False))
++AxisConfig=(AxisKeyName="ValveIndex_Left_Grip_Axis",AxisProperties=(DeadZone=0.000000,Sensitivity=1.000000,Exponent=1.000000,bInvert=False))
 bAltEnterTogglesFullscreen=True
 bF11TogglesFullscreen=True
 bUseMouseForTouch=False
@@ -62,8 +125,8 @@ DefaultViewportMouseCaptureMode=CapturePermanently_IncludingInitialMouseDown
 DefaultViewportMouseLockMode=LockOnCapture
 FOVScale=0.011110
 DoubleClickTime=0.200000
-+ActionMappings=(ActionName="MotionControllerL",bShift=False,bCtrl=False,bAlt=False,bCmd=False,Key=MotionController_Left_Trigger)
-+ActionMappings=(ActionName="MotionControllerR",bShift=False,bCtrl=False,bAlt=False,bCmd=False,Key=MotionController_Right_Trigger)
++ActionMappings=(ActionName="VR_Trigger_L",bShift=False,bCtrl=False,bAlt=False,bCmd=False,Key=OculusTouch_Left_Trigger_Click)
++ActionMappings=(ActionName="VR_Trigger_R",bShift=False,bCtrl=False,bAlt=False,bCmd=False,Key=OculusTouch_Right_Trigger_Click)
 +ActionMappings=(ActionName="PushAndPull",bShift=False,bCtrl=False,bAlt=False,bCmd=False,Key=MouseWheelAxis)
 +ActionMappings=(ActionName="Quit",bShift=False,bCtrl=False,bAlt=False,bCmd=False,Key=Escape)
 +ActionMappings=(ActionName="FlyMode",bShift=False,bCtrl=True,bAlt=False,bCmd=False,Key=One)
@@ -88,6 +151,34 @@ DoubleClickTime=0.200000
 +ActionMappings=(ActionName="RotateView",bShift=False,bCtrl=False,bAlt=False,bCmd=False,Key=RightMouseButton)
 +ActionMappings=(ActionName="SetPointOfinterest",bShift=False,bCtrl=False,bAlt=False,bCmd=False,Key=MiddleMouseButton)
 +ActionMappings=(ActionName="LockToAxis",bShift=False,bCtrl=False,bAlt=False,bCmd=False,Key=LeftShift)
++ActionMappings=(ActionName="ThumbStick_L_Y",bShift=False,bCtrl=False,bAlt=False,bCmd=False,Key=OculusTouch_Left_Thumbstick_Y)
++ActionMappings=(ActionName="ThumbStick_R_Y",bShift=False,bCtrl=False,bAlt=False,bCmd=False,Key=OculusTouch_Right_Thumbstick_Y)
++ActionMappings=(ActionName="ThumbStick_L_X",bShift=False,bCtrl=False,bAlt=False,bCmd=False,Key=OculusTouch_Left_Thumbstick_X)
++ActionMappings=(ActionName="ShowContextualMenu_R",bShift=False,bCtrl=False,bAlt=False,bCmd=False,Key=OculusTouch_Right_A_Click)
++ActionMappings=(ActionName="ShowContextualMenu_L",bShift=False,bCtrl=False,bAlt=False,bCmd=False,Key=OculusTouch_Left_X_Click)
++ActionMappings=(ActionName="VR_Teleport_L",bShift=False,bCtrl=False,bAlt=False,bCmd=False,Key=OculusTouch_Left_Y_Click)
++ActionMappings=(ActionName="VR_Teleport_R",bShift=False,bCtrl=False,bAlt=False,bCmd=False,Key=OculusTouch_Right_B_Click)
++ActionMappings=(ActionName="ShowContextualMenu_R",bShift=False,bCtrl=False,bAlt=False,bCmd=False,Key=Vive_Right_Menu_Click)
++ActionMappings=(ActionName="ThumbStick_L_Y",bShift=False,bCtrl=False,bAlt=False,bCmd=False,Key=Vive_Left_Trackpad_Y)
++ActionMappings=(ActionName="ThumbStick_R_Y",bShift=False,bCtrl=False,bAlt=False,bCmd=False,Key=Vive_Right_Trackpad_Y)
++ActionMappings=(ActionName="ThumbStick_L_X",bShift=False,bCtrl=False,bAlt=False,bCmd=False,Key=Vive_Left_Trackpad_X)
++ActionMappings=(ActionName="VR_Teleport_L",bShift=False,bCtrl=False,bAlt=False,bCmd=False,Key=Vive_Left_Grip_Click)
++ActionMappings=(ActionName="VR_Teleport_R",bShift=False,bCtrl=False,bAlt=False,bCmd=False,Key=Vive_Right_Grip_Click)
++ActionMappings=(ActionName="ShowContextualMenu_L",bShift=False,bCtrl=False,bAlt=False,bCmd=False,Key=Vive_Left_Menu_Click)
++ActionMappings=(ActionName="ThumbStick_R_X",bShift=False,bCtrl=False,bAlt=False,bCmd=False,Key=OculusTouch_Right_Thumbstick_X)
++ActionMappings=(ActionName="VR_Trigger_L",bShift=False,bCtrl=False,bAlt=False,bCmd=False,Key=Vive_Left_Trigger_Click)
++ActionMappings=(ActionName="VR_Trigger_R",bShift=False,bCtrl=False,bAlt=False,bCmd=False,Key=Vive_Right_Trigger_Click)
++ActionMappings=(ActionName="VR_Trigger_L",bShift=False,bCtrl=False,bAlt=False,bCmd=False,Key=ValveIndex_Left_Trigger_Click)
++ActionMappings=(ActionName="VR_Trigger_R",bShift=False,bCtrl=False,bAlt=False,bCmd=False,Key=ValveIndex_Right_Trigger_Click)
++ActionMappings=(ActionName="ThumbStick_L_Y",bShift=False,bCtrl=False,bAlt=False,bCmd=False,Key=ValveIndex_Left_Thumbstick_Y)
++ActionMappings=(ActionName="ThumbStick_R_Y",bShift=False,bCtrl=False,bAlt=False,bCmd=False,Key=ValveIndex_Right_Thumbstick_Y)
++ActionMappings=(ActionName="ThumbStick_L_X",bShift=False,bCtrl=False,bAlt=False,bCmd=False,Key=ValveIndex_Left_Thumbstick_X)
++ActionMappings=(ActionName="ThumbStick_R_X",bShift=False,bCtrl=False,bAlt=False,bCmd=False,Key=Vive_Right_Trackpad_X)
++ActionMappings=(ActionName="ThumbStick_R_X",bShift=False,bCtrl=False,bAlt=False,bCmd=False,Key=ValveIndex_Right_Thumbstick_X)
++ActionMappings=(ActionName="ShowContextualMenu_R",bShift=False,bCtrl=False,bAlt=False,bCmd=False,Key=ValveIndex_Right_A_Click)
++ActionMappings=(ActionName="ShowContextualMenu_L",bShift=False,bCtrl=False,bAlt=False,bCmd=False,Key=ValveIndex_Left_A_Click)
++ActionMappings=(ActionName="VR_Teleport_L",bShift=False,bCtrl=False,bAlt=False,bCmd=False,Key=ValveIndex_Left_B_Click)
++ActionMappings=(ActionName="VR_Teleport_R",bShift=False,bCtrl=False,bAlt=False,bCmd=False,Key=ValveIndex_Right_B_Click)
 +AxisMappings=(AxisName="WalkForward",Scale=1.000000,Key=W)
 +AxisMappings=(AxisName="WalkRight",Scale=1.000000,Key=D)
 +AxisMappings=(AxisName="WalkForward",Scale=-1.000000,Key=S)
@@ -96,10 +187,6 @@ DoubleClickTime=0.200000
 +AxisMappings=(AxisName="LookRight",Scale=1.000000,Key=MouseX)
 +AxisMappings=(AxisName="FlyUp",Scale=1.000000,Key=E)
 +AxisMappings=(AxisName="FlyUp",Scale=-1.000000,Key=Q)
-+AxisMappings=(AxisName="MotionControllerThumbLeft_Y",Scale=-1.000000,Key=MotionController_Left_Thumbstick_Y)
-+AxisMappings=(AxisName="MotionControllerThumbRight_Y",Scale=-1.000000,Key=MotionController_Right_Thumbstick_Y)
-+AxisMappings=(AxisName="MotionControllerThumbLeft_X",Scale=1.000000,Key=MotionController_Left_Thumbstick_X)
-+AxisMappings=(AxisName="MotionControllerThumbRight_X",Scale=1.000000,Key=MotionController_Right_Thumbstick_X)
 +AxisMappings=(AxisName="WalkForward",Scale=1.000000,Key=Gamepad_LeftY)
 +AxisMappings=(AxisName="WalkRight",Scale=1.000000,Key=Gamepad_LeftX)
 +AxisMappings=(AxisName="LookUp",Scale=-1.000000,Key=Gamepad_RightY)
@@ -109,7 +196,20 @@ DoubleClickTime=0.200000
 +AxisMappings=(AxisName="WalkForward",Scale=-1.000000,Key=Down)
 +AxisMappings=(AxisName="WalkRight",Scale=1.000000,Key=Right)
 +AxisMappings=(AxisName="WalkRight",Scale=-1.000000,Key=Left)
++AxisMappings=(AxisName="ThumbStick_L_X",Scale=1.000000,Key=OculusTouch_Left_Thumbstick_X)
++AxisMappings=(AxisName="ThumbStick_R_X",Scale=1.000000,Key=OculusTouch_Right_Thumbstick_X)
++AxisMappings=(AxisName="ThumbStick_L_Y",Scale=-1.000000,Key=OculusTouch_Left_Thumbstick_Y)
++AxisMappings=(AxisName="ThumbStick_R_Y",Scale=-1.000000,Key=OculusTouch_Right_Thumbstick_Y)
++AxisMappings=(AxisName="ThumbStick_L_X",Scale=1.000000,Key=Vive_Left_Trackpad_X)
++AxisMappings=(AxisName="ThumbStick_R_X",Scale=1.000000,Key=Vive_Right_Trackpad_X)
++AxisMappings=(AxisName="ThumbStick_L_Y",Scale=-1.000000,Key=Vive_Left_Trackpad_Y)
++AxisMappings=(AxisName="ThumbStick_R_Y",Scale=-1.000000,Key=Vive_Right_Trackpad_Y)
++AxisMappings=(AxisName="ThumbStick_L_X",Scale=1.000000,Key=ValveIndex_Left_Thumbstick_X)
++AxisMappings=(AxisName="ThumbStick_R_X",Scale=1.000000,Key=ValveIndex_Right_Thumbstick_X)
++AxisMappings=(AxisName="ThumbStick_L_Y",Scale=-1.000000,Key=ValveIndex_Left_Thumbstick_Y)
++AxisMappings=(AxisName="ThumbStick_R_Y",Scale=-1.000000,Key=ValveIndex_Right_Thumbstick_Y)
 DefaultTouchInterface=/Game/CollaborativeViewer/Blueprints/Pawn/Touch/TouchInterface.TouchInterface
 ConsoleKey=None
 +ConsoleKeys=Tilde
 
+
