Fix path to get contents from Collaborative Viewer template from the FeaturePack

diff --git a/Templates/TP_CollaborativeBP/Config/config.ini b/Templates/TP_CollaborativeBP/Config/config.ini
--- a/Templates/TP_CollaborativeBP/Config/config.ini
+++ b/Templates/TP_CollaborativeBP/Config/config.ini
@@ -215,4 +215,4 @@ ConsoleKey=None
 GameName=TP_CollaborativeBP
 
 [AdditionalFilesToAdd]
-+Files=Enterprise/Templates/TP_CollaborativeBP/Content/*.*
++Files=Templates/TP_CollaborativeBP/Content/*.*
