Collaborative viewer - Menu navigation is not working with Vive hardware

diff --git a/Templates/TP_CollaborativeBP/Config/DefaultInput.ini b/Templates/TP_CollaborativeBP/Config/DefaultInput.ini
--- a/Templates/TP_CollaborativeBP/Config/DefaultInput.ini
+++ b/Templates/TP_CollaborativeBP/Config/DefaultInput.ini
@@ -1,5 +1,23 @@
 
 [/Script/Engine.InputSettings]
++AxisConfig=(AxisKeyName="MixedReality_Right_Thumbstick_Y",AxisProperties=(DeadZone=0.000000,Sensitivity=1.000000,Exponent=1.000000,bInvert=False))
++AxisConfig=(AxisKeyName="MixedReality_Right_Trackpad_X",AxisProperties=(DeadZone=0.000000,Sensitivity=1.000000,Exponent=1.000000,bInvert=False))
++AxisConfig=(AxisKeyName="MixedReality_Right_Trackpad_Y",AxisProperties=(DeadZone=0.000000,Sensitivity=1.000000,Exponent=1.000000,bInvert=False))
++AxisConfig=(AxisKeyName="OculusGo_Left_Trackpad_X",AxisProperties=(DeadZone=0.000000,Sensitivity=1.000000,Exponent=1.000000,bInvert=False))
++AxisConfig=(AxisKeyName="Cosmos_Left_Grip_Axis",AxisProperties=(DeadZone=0.000000,Sensitivity=1.000000,Exponent=1.000000,bInvert=False))
++AxisConfig=(AxisKeyName="Cosmos_Left_Trigger_Axis",AxisProperties=(DeadZone=0.000000,Sensitivity=1.000000,Exponent=1.000000,bInvert=False))
++AxisConfig=(AxisKeyName="Cosmos_Left_Thumbstick_X",AxisProperties=(DeadZone=0.000000,Sensitivity=1.000000,Exponent=1.000000,bInvert=False))
++AxisConfig=(AxisKeyName="Cosmos_Left_Thumbstick_Y",AxisProperties=(DeadZone=0.000000,Sensitivity=1.000000,Exponent=1.000000,bInvert=False))
++AxisConfig=(AxisKeyName="Cosmos_Right_Grip_Axis",AxisProperties=(DeadZone=0.000000,Sensitivity=1.000000,Exponent=1.000000,bInvert=False))
++AxisConfig=(AxisKeyName="Cosmos_Right_Trigger_Axis",AxisProperties=(DeadZone=0.000000,Sensitivity=1.000000,Exponent=1.000000,bInvert=False))
++AxisConfig=(AxisKeyName="Cosmos_Right_Thumbstick_X",AxisProperties=(DeadZone=0.000000,Sensitivity=1.000000,Exponent=1.000000,bInvert=False))
++AxisConfig=(AxisKeyName="Cosmos_Right_Thumbstick_Y",AxisProperties=(DeadZone=0.000000,Sensitivity=1.000000,Exponent=1.000000,bInvert=False))
++AxisConfig=(AxisKeyName="OculusGo_Left_Trackpad_Y",AxisProperties=(DeadZone=0.000000,Sensitivity=1.000000,Exponent=1.000000,bInvert=False))
++AxisConfig=(AxisKeyName="OculusGo_Right_Trackpad_X",AxisProperties=(DeadZone=0.000000,Sensitivity=1.000000,Exponent=1.000000,bInvert=False))
++AxisConfig=(AxisKeyName="OculusGo_Right_Trackpad_Y",AxisProperties=(DeadZone=0.000000,Sensitivity=1.000000,Exponent=1.000000,bInvert=False))
++AxisConfig=(AxisKeyName="OculusTouch_Left_Grip_Axis",AxisProperties=(DeadZone=0.000000,Sensitivity=1.000000,Exponent=1.000000,bInvert=False))
++AxisConfig=(AxisKeyName="OculusTouch_Left_Trigger_Axis",AxisProperties=(DeadZone=0.000000,Sensitivity=1.000000,Exponent=1.000000,bInvert=False))
++AxisConfig=(AxisKeyName="OculusTouch_Left_Thumbstick_X",AxisProperties=(DeadZone=0.000000,Sensitivity=1.000000,Exponent=1.000000,bInvert=False))
 +AxisConfig=(AxisKeyName="OculusTouch_Left_Thumbstick_Y",AxisProperties=(DeadZone=0.000000,Sensitivity=1.000000,Exponent=1.000000,bInvert=False))
 +AxisConfig=(AxisKeyName="OculusTouch_Right_Grip_Axis",AxisProperties=(DeadZone=0.000000,Sensitivity=1.000000,Exponent=1.000000,bInvert=False))
 +AxisConfig=(AxisKeyName="OculusTouch_Right_Trigger_Axis",AxisProperties=(DeadZone=0.000000,Sensitivity=1.000000,Exponent=1.000000,bInvert=False))
@@ -100,16 +118,6 @@
 +AxisConfig=(AxisKeyName="MixedReality_Left_Trackpad_Y",AxisProperties=(DeadZone=0.000000,Sensitivity=1.000000,Exponent=1.000000,bInvert=False))
 +AxisConfig=(AxisKeyName="MixedReality_Right_Trigger_Axis",AxisProperties=(DeadZone=0.000000,Sensitivity=1.000000,Exponent=1.000000,bInvert=False))
 +AxisConfig=(AxisKeyName="MixedReality_Right_Thumbstick_X",AxisProperties=(DeadZone=0.000000,Sensitivity=1.000000,Exponent=1.000000,bInvert=False))
-+AxisConfig=(AxisKeyName="MixedReality_Right_Thumbstick_Y",AxisProperties=(DeadZone=0.000000,Sensitivity=1.000000,Exponent=1.000000,bInvert=False))
-+AxisConfig=(AxisKeyName="MixedReality_Right_Trackpad_X",AxisProperties=(DeadZone=0.000000,Sensitivity=1.000000,Exponent=1.000000,bInvert=False))
-+AxisConfig=(AxisKeyName="MixedReality_Right_Trackpad_Y",AxisProperties=(DeadZone=0.000000,Sensitivity=1.000000,Exponent=1.000000,bInvert=False))
-+AxisConfig=(AxisKeyName="OculusGo_Left_Trackpad_X",AxisProperties=(DeadZone=0.000000,Sensitivity=1.000000,Exponent=1.000000,bInvert=False))
-+AxisConfig=(AxisKeyName="OculusGo_Left_Trackpad_Y",AxisProperties=(DeadZone=0.000000,Sensitivity=1.000000,Exponent=1.000000,bInvert=False))
-+AxisConfig=(AxisKeyName="OculusGo_Right_Trackpad_X",AxisProperties=(DeadZone=0.000000,Sensitivity=1.000000,Exponent=1.000000,bInvert=False))
-+AxisConfig=(AxisKeyName="OculusGo_Right_Trackpad_Y",AxisProperties=(DeadZone=0.000000,Sensitivity=1.000000,Exponent=1.000000,bInvert=False))
-+AxisConfig=(AxisKeyName="OculusTouch_Left_Grip_Axis",AxisProperties=(DeadZone=0.000000,Sensitivity=1.000000,Exponent=1.000000,bInvert=False))
-+AxisConfig=(AxisKeyName="OculusTouch_Left_Trigger_Axis",AxisProperties=(DeadZone=0.000000,Sensitivity=1.000000,Exponent=1.000000,bInvert=False))
-+AxisConfig=(AxisKeyName="OculusTouch_Left_Thumbstick_X",AxisProperties=(DeadZone=0.000000,Sensitivity=1.000000,Exponent=1.000000,bInvert=False))
 bAltEnterTogglesFullscreen=True
 bF11TogglesFullscreen=True
 bUseMouseForTouch=False
@@ -150,36 +158,48 @@ DoubleClickTime=0.200000
 +ActionMappings=(ActionName="RotateView",bShift=False,bCtrl=False,bAlt=False,bCmd=False,Key=RightMouseButton)
 +ActionMappings=(ActionName="SetPointOfinterest",bShift=False,bCtrl=False,bAlt=False,bCmd=False,Key=MiddleMouseButton)
 +ActionMappings=(ActionName="LockToAxis",bShift=False,bCtrl=False,bAlt=False,bCmd=False,Key=LeftShift)
-+ActionMappings=(ActionName="ThumbStick_L_Y",bShift=False,bCtrl=False,bAlt=False,bCmd=False,Key=OculusTouch_Left_Thumbstick_Y)
-+ActionMappings=(ActionName="ThumbStick_R_Y",bShift=False,bCtrl=False,bAlt=False,bCmd=False,Key=OculusTouch_Right_Thumbstick_Y)
-+ActionMappings=(ActionName="ThumbStick_L_X",bShift=False,bCtrl=False,bAlt=False,bCmd=False,Key=OculusTouch_Left_Thumbstick_X)
 +ActionMappings=(ActionName="ShowContextualMenu_R",bShift=False,bCtrl=False,bAlt=False,bCmd=False,Key=OculusTouch_Right_A_Click)
 +ActionMappings=(ActionName="ShowContextualMenu_L",bShift=False,bCtrl=False,bAlt=False,bCmd=False,Key=OculusTouch_Left_X_Click)
 +ActionMappings=(ActionName="VR_Teleport_L",bShift=False,bCtrl=False,bAlt=False,bCmd=False,Key=OculusTouch_Left_Y_Click)
 +ActionMappings=(ActionName="VR_Teleport_R",bShift=False,bCtrl=False,bAlt=False,bCmd=False,Key=OculusTouch_Right_B_Click)
 +ActionMappings=(ActionName="ShowContextualMenu_R",bShift=False,bCtrl=False,bAlt=False,bCmd=False,Key=Vive_Right_Menu_Click)
-+ActionMappings=(ActionName="ThumbStick_L_Y",bShift=False,bCtrl=False,bAlt=False,bCmd=False,Key=Vive_Left_Trackpad_Y)
-+ActionMappings=(ActionName="ThumbStick_R_Y",bShift=False,bCtrl=False,bAlt=False,bCmd=False,Key=Vive_Right_Trackpad_Y)
-+ActionMappings=(ActionName="ThumbStick_L_X",bShift=False,bCtrl=False,bAlt=False,bCmd=False,Key=Vive_Left_Trackpad_X)
 +ActionMappings=(ActionName="VR_Teleport_L",bShift=False,bCtrl=False,bAlt=False,bCmd=False,Key=Vive_Left_Grip_Click)
 +ActionMappings=(ActionName="VR_Teleport_R",bShift=False,bCtrl=False,bAlt=False,bCmd=False,Key=Vive_Right_Grip_Click)
 +ActionMappings=(ActionName="ShowContextualMenu_L",bShift=False,bCtrl=False,bAlt=False,bCmd=False,Key=Vive_Left_Menu_Click)
-+ActionMappings=(ActionName="ThumbStick_R_X",bShift=False,bCtrl=False,bAlt=False,bCmd=False,Key=OculusTouch_Right_Thumbstick_X)
 +ActionMappings=(ActionName="VR_Trigger_L",bShift=False,bCtrl=False,bAlt=False,bCmd=False,Key=Vive_Left_Trigger_Click)
 +ActionMappings=(ActionName="VR_Trigger_R",bShift=False,bCtrl=False,bAlt=False,bCmd=False,Key=Vive_Right_Trigger_Click)
 +ActionMappings=(ActionName="VR_Trigger_L",bShift=False,bCtrl=False,bAlt=False,bCmd=False,Key=ValveIndex_Left_Trigger_Click)
 +ActionMappings=(ActionName="VR_Trigger_R",bShift=False,bCtrl=False,bAlt=False,bCmd=False,Key=ValveIndex_Right_Trigger_Click)
-+ActionMappings=(ActionName="ThumbStick_L_Y",bShift=False,bCtrl=False,bAlt=False,bCmd=False,Key=ValveIndex_Left_Thumbstick_Y)
-+ActionMappings=(ActionName="ThumbStick_R_Y",bShift=False,bCtrl=False,bAlt=False,bCmd=False,Key=ValveIndex_Right_Thumbstick_Y)
-+ActionMappings=(ActionName="ThumbStick_L_X",bShift=False,bCtrl=False,bAlt=False,bCmd=False,Key=ValveIndex_Left_Thumbstick_X)
-+ActionMappings=(ActionName="ThumbStick_R_X",bShift=False,bCtrl=False,bAlt=False,bCmd=False,Key=Vive_Right_Trackpad_X)
-+ActionMappings=(ActionName="ThumbStick_R_X",bShift=False,bCtrl=False,bAlt=False,bCmd=False,Key=ValveIndex_Right_Thumbstick_X)
 +ActionMappings=(ActionName="ShowContextualMenu_R",bShift=False,bCtrl=False,bAlt=False,bCmd=False,Key=ValveIndex_Right_A_Click)
 +ActionMappings=(ActionName="ShowContextualMenu_L",bShift=False,bCtrl=False,bAlt=False,bCmd=False,Key=ValveIndex_Left_A_Click)
 +ActionMappings=(ActionName="VR_Teleport_L",bShift=False,bCtrl=False,bAlt=False,bCmd=False,Key=ValveIndex_Left_B_Click)
 +ActionMappings=(ActionName="VR_Teleport_R",bShift=False,bCtrl=False,bAlt=False,bCmd=False,Key=ValveIndex_Right_B_Click)
 +ActionMappings=(ActionName="ResetLocationL",bShift=False,bCtrl=False,bAlt=False,bCmd=False,Key=Vive_Left_Grip_Click)
 +ActionMappings=(ActionName="ResetLocationL",bShift=False,bCtrl=False,bAlt=False,bCmd=False,Key=ValveIndex_Left_Thumbstick_Click)
++ActionMappings=(ActionName="UP_R",bShift=False,bCtrl=False,bAlt=False,bCmd=False,Key=Vive_Right_Trackpad_Up)
++ActionMappings=(ActionName="UP_L",bShift=False,bCtrl=False,bAlt=False,bCmd=False,Key=Vive_Left_Trackpad_Up)
++ActionMappings=(ActionName="Down_R",bShift=False,bCtrl=False,bAlt=False,bCmd=False,Key=Vive_Right_Trackpad_Down)
++ActionMappings=(ActionName="Down_L",bShift=False,bCtrl=False,bAlt=False,bCmd=False,Key=Vive_Left_Trackpad_Down)
++ActionMappings=(ActionName="Right_R",bShift=False,bCtrl=False,bAlt=False,bCmd=False,Key=Vive_Right_Trackpad_Right)
++ActionMappings=(ActionName="Right_L",bShift=False,bCtrl=False,bAlt=False,bCmd=False,Key=Vive_Left_Trackpad_Right)
++ActionMappings=(ActionName="Left_R",bShift=False,bCtrl=False,bAlt=False,bCmd=False,Key=Vive_Right_Trackpad_Left)
++ActionMappings=(ActionName="Left_L",bShift=False,bCtrl=False,bAlt=False,bCmd=False,Key=Vive_Left_Trackpad_Left)
++ActionMappings=(ActionName="UP_R",bShift=False,bCtrl=False,bAlt=False,bCmd=False,Key=OculusTouch_Right_Thumbstick_Up)
++ActionMappings=(ActionName="UP_L",bShift=False,bCtrl=False,bAlt=False,bCmd=False,Key=OculusTouch_Left_Thumbstick_Up)
++ActionMappings=(ActionName="Down_R",bShift=False,bCtrl=False,bAlt=False,bCmd=False,Key=OculusTouch_Right_Thumbstick_Down)
++ActionMappings=(ActionName="Down_L",bShift=False,bCtrl=False,bAlt=False,bCmd=False,Key=OculusTouch_Left_Thumbstick_Down)
++ActionMappings=(ActionName="Right_R",bShift=False,bCtrl=False,bAlt=False,bCmd=False,Key=OculusTouch_Right_Thumbstick_Right)
++ActionMappings=(ActionName="Right_L",bShift=False,bCtrl=False,bAlt=False,bCmd=False,Key=OculusTouch_Left_Thumbstick_Right)
++ActionMappings=(ActionName="Left_R",bShift=False,bCtrl=False,bAlt=False,bCmd=False,Key=OculusTouch_Right_Thumbstick_Left)
++ActionMappings=(ActionName="Left_L",bShift=False,bCtrl=False,bAlt=False,bCmd=False,Key=OculusTouch_Left_Thumbstick_Left)
++ActionMappings=(ActionName="UP_R",bShift=False,bCtrl=False,bAlt=False,bCmd=False,Key=ValveIndex_Right_Thumbstick_Up)
++ActionMappings=(ActionName="UP_L",bShift=False,bCtrl=False,bAlt=False,bCmd=False,Key=ValveIndex_Left_Thumbstick_Up)
++ActionMappings=(ActionName="Down_R",bShift=False,bCtrl=False,bAlt=False,bCmd=False,Key=ValveIndex_Right_Thumbstick_Down)
++ActionMappings=(ActionName="Down_L",bShift=False,bCtrl=False,bAlt=False,bCmd=False,Key=ValveIndex_Left_Thumbstick_Down)
++ActionMappings=(ActionName="Right_R",bShift=False,bCtrl=False,bAlt=False,bCmd=False,Key=ValveIndex_Right_Thumbstick_Right)
++ActionMappings=(ActionName="Right_L",bShift=False,bCtrl=False,bAlt=False,bCmd=False,Key=ValveIndex_Left_Thumbstick_Right)
++ActionMappings=(ActionName="Left_R",bShift=False,bCtrl=False,bAlt=False,bCmd=False,Key=ValveIndex_Right_Thumbstick_Left)
++ActionMappings=(ActionName="Left_L",bShift=False,bCtrl=False,bAlt=False,bCmd=False,Key=ValveIndex_Left_Thumbstick_Left)
 +AxisMappings=(AxisName="WalkForward",Scale=1.000000,Key=W)
 +AxisMappings=(AxisName="WalkRight",Scale=1.000000,Key=D)
 +AxisMappings=(AxisName="WalkForward",Scale=-1.000000,Key=S)
@@ -197,18 +217,10 @@ DoubleClickTime=0.200000
 +AxisMappings=(AxisName="WalkForward",Scale=-1.000000,Key=Down)
 +AxisMappings=(AxisName="WalkRight",Scale=1.000000,Key=Right)
 +AxisMappings=(AxisName="WalkRight",Scale=-1.000000,Key=Left)
-+AxisMappings=(AxisName="ThumbStick_L_X",Scale=1.000000,Key=OculusTouch_Left_Thumbstick_X)
-+AxisMappings=(AxisName="ThumbStick_R_X",Scale=1.000000,Key=OculusTouch_Right_Thumbstick_X)
-+AxisMappings=(AxisName="ThumbStick_L_Y",Scale=-1.000000,Key=OculusTouch_Left_Thumbstick_Y)
-+AxisMappings=(AxisName="ThumbStick_R_Y",Scale=-1.000000,Key=OculusTouch_Right_Thumbstick_Y)
-+AxisMappings=(AxisName="ThumbStick_L_X",Scale=1.000000,Key=Vive_Left_Trackpad_X)
-+AxisMappings=(AxisName="ThumbStick_R_X",Scale=1.000000,Key=Vive_Right_Trackpad_X)
-+AxisMappings=(AxisName="ThumbStick_L_Y",Scale=-1.000000,Key=Vive_Left_Trackpad_Y)
-+AxisMappings=(AxisName="ThumbStick_R_Y",Scale=-1.000000,Key=Vive_Right_Trackpad_Y)
-+AxisMappings=(AxisName="ThumbStick_L_X",Scale=1.000000,Key=ValveIndex_Left_Thumbstick_X)
-+AxisMappings=(AxisName="ThumbStick_R_X",Scale=1.000000,Key=ValveIndex_Right_Thumbstick_X)
-+AxisMappings=(AxisName="ThumbStick_L_Y",Scale=-1.000000,Key=ValveIndex_Left_Thumbstick_Y)
-+AxisMappings=(AxisName="ThumbStick_R_Y",Scale=-1.000000,Key=ValveIndex_Right_Thumbstick_Y)
++AxisMappings=(AxisName="Thumbstick_R_Y",Scale=1.000000,Key=OculusTouch_Right_Thumbstick_Y)
++AxisMappings=(AxisName="Thumbstick_L_Y",Scale=1.000000,Key=OculusTouch_Left_Thumbstick_Y)
++AxisMappings=(AxisName="Thumbstick_R_Y",Scale=1.000000,Key=ValveIndex_Right_Thumbstick_Y)
++AxisMappings=(AxisName="Thumbstick_L_Y",Scale=1.000000,Key=ValveIndex_Left_Thumbstick_Y)
 DefaultTouchInterface=/Game/CollaborativeViewer/Blueprints/Pawn/Touch/TouchInterface.TouchInterface
 +ConsoleKeys=Tilde
 
