Collaborative viewer - Menu navigation is not working with Vive hardware
Those file are needed for the vive Bindings if you dont have a steam VR install on your desktop

diff --git a/Templates/TP_CollaborativeBP/Config/SteamVRBindings/gamepad.json b/Templates/TP_CollaborativeBP/Config/SteamVRBindings/gamepad.json
new file mode 100644
--- /dev/null
+++ b/Templates/TP_CollaborativeBP/Config/SteamVRBindings/gamepad.json
@@ -0,0 +1,13 @@
+{
+	"name": "Default bindings for Gamepads",
+	"controller_type": "gamepad",
+	"last_edited_by": "UnrealEngine",
+	"bindings":
+	{
+		"/actions/main":
+		{
+			"sources": []
+		}
+	},
+	"description": "TP_CollaborativeBP-10452296"
+}
\ No newline at end of file
diff --git a/Templates/TP_CollaborativeBP/Config/SteamVRBindings/holographic_controller.json b/Templates/TP_CollaborativeBP/Config/SteamVRBindings/holographic_controller.json
new file mode 100644
--- /dev/null
+++ b/Templates/TP_CollaborativeBP/Config/SteamVRBindings/holographic_controller.json
@@ -0,0 +1,44 @@
+{
+	"name": "Default bindings for MixedReality",
+	"controller_type": "holographic_controller",
+	"last_edited_by": "UnrealEngine",
+	"bindings":
+	{
+		"/actions/main":
+		{
+			"sources": [],
+			"poses": [
+				{
+					"output": "/actions/main/in/controllerleft",
+					"path": "/user/hand/left/pose/raw",
+					"requirement": "optional"
+				},
+				{
+					"output": "/actions/main/in/controllerright",
+					"path": "/user/hand/right/pose/raw"
+				}
+			],
+			"skeleton": [
+				{
+					"output": "/actions/main/in/skeletonleft",
+					"path": "/user/hand/left/input/skeleton/left"
+				},
+				{
+					"output": "/actions/main/in/skeletonright",
+					"path": "/user/hand/right/input/skeleton/right"
+				}
+			],
+			"haptics": [
+				{
+					"output": "/actions/main/out/vibrateleft",
+					"path": "/user/hand/left/output/haptic"
+				},
+				{
+					"output": "/actions/main/out/vibrateright",
+					"path": "/user/hand/right/output/haptic"
+				}
+			]
+		}
+	},
+	"description": "TP_CollaborativeBP-10452296"
+}
\ No newline at end of file
diff --git a/Templates/TP_CollaborativeBP/Config/SteamVRBindings/indexhmd.json b/Templates/TP_CollaborativeBP/Config/SteamVRBindings/indexhmd.json
new file mode 100644
--- /dev/null
+++ b/Templates/TP_CollaborativeBP/Config/SteamVRBindings/indexhmd.json
@@ -0,0 +1,13 @@
+{
+	"name": "Default bindings for Valve Index Headset",
+	"controller_type": "indexhmd",
+	"last_edited_by": "UnrealEngine",
+	"bindings":
+	{
+		"/actions/main":
+		{
+			"sources": []
+		}
+	},
+	"description": "TP_CollaborativeBP-10452296"
+}
\ No newline at end of file
diff --git a/Templates/TP_CollaborativeBP/Config/SteamVRBindings/knuckles.json b/Templates/TP_CollaborativeBP/Config/SteamVRBindings/knuckles.json
new file mode 100644
--- /dev/null
+++ b/Templates/TP_CollaborativeBP/Config/SteamVRBindings/knuckles.json
@@ -0,0 +1,242 @@
+{
+	"name": "Default bindings for ValveIndex",
+	"controller_type": "knuckles",
+	"last_edited_by": "UnrealEngine",
+	"bindings":
+	{
+		"/actions/main":
+		{
+			"sources": [
+				{
+					"mode": "button",
+					"path": "/user/hand/left/input/trigger",
+					"inputs":
+					{
+						"click":
+						{
+							"output": "/actions/main/in/VR_Trigger_L"
+						}
+					}
+				},
+				{
+					"mode": "button",
+					"path": "/user/hand/right/input/trigger",
+					"inputs":
+					{
+						"click":
+						{
+							"output": "/actions/main/in/VR_Trigger_R"
+						}
+					}
+				},
+				{
+					"mode": "joystick",
+					"path": "/user/hand/left/input/thumbstick",
+					"inputs":
+					{
+						"click":
+						{
+							"output": "/actions/main/in/ResetLocationL"
+						}
+					}
+				},
+				{
+					"mode": "button",
+					"path": "/user/hand/right/input/a",
+					"inputs":
+					{
+						"click":
+						{
+							"output": "/actions/main/in/ShowContextualMenu_R"
+						}
+					}
+				},
+				{
+					"mode": "button",
+					"path": "/user/hand/left/input/a",
+					"inputs":
+					{
+						"click":
+						{
+							"output": "/actions/main/in/ShowContextualMenu_L"
+						}
+					}
+				},
+				{
+					"mode": "button",
+					"path": "/user/hand/left/input/b",
+					"inputs":
+					{
+						"click":
+						{
+							"output": "/actions/main/in/VR_Teleport_L"
+						}
+					}
+				},
+				{
+					"mode": "button",
+					"path": "/user/hand/right/input/b",
+					"inputs":
+					{
+						"click":
+						{
+							"output": "/actions/main/in/VR_Teleport_R"
+						}
+					}
+				},
+				{
+					"mode": "dpad",
+					"path": "/user/hand/right/input/thumbstick",
+					"parameters":
+					{
+						"sub_mode": "click"
+					},
+					"inputs":
+					{
+						"north":
+						{
+							"output": "/actions/main/in/UP_R"
+						}
+					}
+				},
+				{
+					"mode": "dpad",
+					"path": "/user/hand/left/input/thumbstick",
+					"parameters":
+					{
+						"sub_mode": "click"
+					},
+					"inputs":
+					{
+						"north":
+						{
+							"output": "/actions/main/in/UP_L"
+						}
+					}
+				},
+				{
+					"mode": "dpad",
+					"path": "/user/hand/right/input/thumbstick",
+					"parameters":
+					{
+						"sub_mode": "click"
+					},
+					"inputs":
+					{
+						"south":
+						{
+							"output": "/actions/main/in/Down_R"
+						}
+					}
+				},
+				{
+					"mode": "dpad",
+					"path": "/user/hand/left/input/thumbstick",
+					"parameters":
+					{
+						"sub_mode": "click"
+					},
+					"inputs":
+					{
+						"south":
+						{
+							"output": "/actions/main/in/Down_L"
+						}
+					}
+				},
+				{
+					"mode": "dpad",
+					"path": "/user/hand/right/input/thumbstick",
+					"parameters":
+					{
+						"sub_mode": "click"
+					},
+					"inputs":
+					{
+						"east":
+						{
+							"output": "/actions/main/in/Right_R"
+						}
+					}
+				},
+				{
+					"mode": "dpad",
+					"path": "/user/hand/left/input/thumbstick",
+					"parameters":
+					{
+						"sub_mode": "click"
+					},
+					"inputs":
+					{
+						"east":
+						{
+							"output": "/actions/main/in/Right_L"
+						}
+					}
+				},
+				{
+					"mode": "dpad",
+					"path": "/user/hand/right/input/thumbstick",
+					"parameters":
+					{
+						"sub_mode": "click"
+					},
+					"inputs":
+					{
+						"west":
+						{
+							"output": "/actions/main/in/Left_R"
+						}
+					}
+				},
+				{
+					"mode": "dpad",
+					"path": "/user/hand/left/input/thumbstick",
+					"parameters":
+					{
+						"sub_mode": "click"
+					},
+					"inputs":
+					{
+						"west":
+						{
+							"output": "/actions/main/in/Left_L"
+						}
+					}
+				}
+			],
+			"poses": [
+				{
+					"output": "/actions/main/in/controllerleft",
+					"path": "/user/hand/left/pose/raw",
+					"requirement": "optional"
+				},
+				{
+					"output": "/actions/main/in/controllerright",
+					"path": "/user/hand/right/pose/raw"
+				}
+			],
+			"skeleton": [
+				{
+					"output": "/actions/main/in/skeletonleft",
+					"path": "/user/hand/left/input/skeleton/left"
+				},
+				{
+					"output": "/actions/main/in/skeletonright",
+					"path": "/user/hand/right/input/skeleton/right"
+				}
+			],
+			"haptics": [
+				{
+					"output": "/actions/main/out/vibrateleft",
+					"path": "/user/hand/left/output/haptic"
+				},
+				{
+					"output": "/actions/main/out/vibrateright",
+					"path": "/user/hand/right/output/haptic"
+				}
+			]
+		}
+	},
+	"description": "TP_CollaborativeBP-10452296"
+}
\ No newline at end of file
diff --git a/Templates/TP_CollaborativeBP/Config/SteamVRBindings/oculus_touch.json b/Templates/TP_CollaborativeBP/Config/SteamVRBindings/oculus_touch.json
new file mode 100644
--- /dev/null
+++ b/Templates/TP_CollaborativeBP/Config/SteamVRBindings/oculus_touch.json
@@ -0,0 +1,242 @@
+{
+	"name": "Default bindings for OculusTouch",
+	"controller_type": "oculus_touch",
+	"last_edited_by": "UnrealEngine",
+	"bindings":
+	{
+		"/actions/main":
+		{
+			"sources": [
+				{
+					"mode": "button",
+					"path": "/user/hand/left/input/trigger",
+					"inputs":
+					{
+						"click":
+						{
+							"output": "/actions/main/in/VR_Trigger_L"
+						}
+					}
+				},
+				{
+					"mode": "button",
+					"path": "/user/hand/right/input/trigger",
+					"inputs":
+					{
+						"click":
+						{
+							"output": "/actions/main/in/VR_Trigger_R"
+						}
+					}
+				},
+				{
+					"mode": "joystick",
+					"path": "/user/hand/left/input/joystick",
+					"inputs":
+					{
+						"click":
+						{
+							"output": "/actions/main/in/ResetLocationL"
+						}
+					}
+				},
+				{
+					"mode": "button",
+					"path": "/user/hand/right/input/a",
+					"inputs":
+					{
+						"click":
+						{
+							"output": "/actions/main/in/ShowContextualMenu_R"
+						}
+					}
+				},
+				{
+					"mode": "button",
+					"path": "/user/hand/left/input/x",
+					"inputs":
+					{
+						"click":
+						{
+							"output": "/actions/main/in/ShowContextualMenu_L"
+						}
+					}
+				},
+				{
+					"mode": "button",
+					"path": "/user/hand/left/input/y",
+					"inputs":
+					{
+						"click":
+						{
+							"output": "/actions/main/in/VR_Teleport_L"
+						}
+					}
+				},
+				{
+					"mode": "button",
+					"path": "/user/hand/right/input/b",
+					"inputs":
+					{
+						"click":
+						{
+							"output": "/actions/main/in/VR_Teleport_R"
+						}
+					}
+				},
+				{
+					"mode": "dpad",
+					"path": "/user/hand/right/input/joystick",
+					"parameters":
+					{
+						"sub_mode": "click"
+					},
+					"inputs":
+					{
+						"north":
+						{
+							"output": "/actions/main/in/UP_R"
+						}
+					}
+				},
+				{
+					"mode": "dpad",
+					"path": "/user/hand/left/input/joystick",
+					"parameters":
+					{
+						"sub_mode": "click"
+					},
+					"inputs":
+					{
+						"north":
+						{
+							"output": "/actions/main/in/UP_L"
+						}
+					}
+				},
+				{
+					"mode": "dpad",
+					"path": "/user/hand/right/input/joystick",
+					"parameters":
+					{
+						"sub_mode": "click"
+					},
+					"inputs":
+					{
+						"south":
+						{
+							"output": "/actions/main/in/Down_R"
+						}
+					}
+				},
+				{
+					"mode": "dpad",
+					"path": "/user/hand/left/input/joystick",
+					"parameters":
+					{
+						"sub_mode": "click"
+					},
+					"inputs":
+					{
+						"south":
+						{
+							"output": "/actions/main/in/Down_L"
+						}
+					}
+				},
+				{
+					"mode": "dpad",
+					"path": "/user/hand/right/input/joystick",
+					"parameters":
+					{
+						"sub_mode": "click"
+					},
+					"inputs":
+					{
+						"east":
+						{
+							"output": "/actions/main/in/Right_R"
+						}
+					}
+				},
+				{
+					"mode": "dpad",
+					"path": "/user/hand/left/input/joystick",
+					"parameters":
+					{
+						"sub_mode": "click"
+					},
+					"inputs":
+					{
+						"east":
+						{
+							"output": "/actions/main/in/Right_L"
+						}
+					}
+				},
+				{
+					"mode": "dpad",
+					"path": "/user/hand/right/input/joystick",
+					"parameters":
+					{
+						"sub_mode": "click"
+					},
+					"inputs":
+					{
+						"west":
+						{
+							"output": "/actions/main/in/Left_R"
+						}
+					}
+				},
+				{
+					"mode": "dpad",
+					"path": "/user/hand/left/input/joystick",
+					"parameters":
+					{
+						"sub_mode": "click"
+					},
+					"inputs":
+					{
+						"west":
+						{
+							"output": "/actions/main/in/Left_L"
+						}
+					}
+				}
+			],
+			"poses": [
+				{
+					"output": "/actions/main/in/controllerleft",
+					"path": "/user/hand/left/pose/raw",
+					"requirement": "optional"
+				},
+				{
+					"output": "/actions/main/in/controllerright",
+					"path": "/user/hand/right/pose/raw"
+				}
+			],
+			"skeleton": [
+				{
+					"output": "/actions/main/in/skeletonleft",
+					"path": "/user/hand/left/input/skeleton/left"
+				},
+				{
+					"output": "/actions/main/in/skeletonright",
+					"path": "/user/hand/right/input/skeleton/right"
+				}
+			],
+			"haptics": [
+				{
+					"output": "/actions/main/out/vibrateleft",
+					"path": "/user/hand/left/output/haptic"
+				},
+				{
+					"output": "/actions/main/out/vibrateright",
+					"path": "/user/hand/right/output/haptic"
+				}
+			]
+		}
+	},
+	"description": "TP_CollaborativeBP-10452296"
+}
\ No newline at end of file
diff --git a/Templates/TP_CollaborativeBP/Config/SteamVRBindings/rift.json b/Templates/TP_CollaborativeBP/Config/SteamVRBindings/rift.json
new file mode 100644
--- /dev/null
+++ b/Templates/TP_CollaborativeBP/Config/SteamVRBindings/rift.json
@@ -0,0 +1,13 @@
+{
+	"name": "Default bindings for Rift Headset",
+	"controller_type": "rift",
+	"last_edited_by": "UnrealEngine",
+	"bindings":
+	{
+		"/actions/main":
+		{
+			"sources": []
+		}
+	},
+	"description": "TP_CollaborativeBP-10452296"
+}
\ No newline at end of file
diff --git a/Templates/TP_CollaborativeBP/Config/SteamVRBindings/steamvr_manifest.json b/Templates/TP_CollaborativeBP/Config/SteamVRBindings/steamvr_manifest.json
new file mode 100644
--- /dev/null
+++ b/Templates/TP_CollaborativeBP/Config/SteamVRBindings/steamvr_manifest.json
@@ -0,0 +1,239 @@
+{
+	"actions": [
+		{
+			"name": "/actions/main/in/controllerleft",
+			"type": "pose",
+			"requirement": "optional"
+		},
+		{
+			"name": "/actions/main/in/controllerright",
+			"type": "pose",
+			"requirement": "optional"
+		},
+		{
+			"name": "/actions/main/in/special1",
+			"type": "pose",
+			"requirement": "optional"
+		},
+		{
+			"name": "/actions/main/in/special2",
+			"type": "pose",
+			"requirement": "optional"
+		},
+		{
+			"name": "/actions/main/in/special3",
+			"type": "pose",
+			"requirement": "optional"
+		},
+		{
+			"name": "/actions/main/in/special4",
+			"type": "pose",
+			"requirement": "optional"
+		},
+		{
+			"name": "/actions/main/in/special5",
+			"type": "pose",
+			"requirement": "optional"
+		},
+		{
+			"name": "/actions/main/in/special6",
+			"type": "pose",
+			"requirement": "optional"
+		},
+		{
+			"name": "/actions/main/in/special7",
+			"type": "pose",
+			"requirement": "optional"
+		},
+		{
+			"name": "/actions/main/in/special8",
+			"type": "pose",
+			"requirement": "optional"
+		},
+		{
+			"name": "/actions/main/in/skeletonleft",
+			"type": "skeleton",
+			"skeleton": "/skeleton/hand/left",
+			"requirement": "optional"
+		},
+		{
+			"name": "/actions/main/in/skeletonright",
+			"type": "skeleton",
+			"skeleton": "/skeleton/hand/right",
+			"requirement": "optional"
+		},
+		{
+			"name": "/actions/main/in/open_console",
+			"type": "boolean",
+			"requirement": "optional"
+		},
+		{
+			"name": "/actions/main/out/vibrateleft",
+			"type": "vibration",
+			"requirement": "optional"
+		},
+		{
+			"name": "/actions/main/out/vibrateright",
+			"type": "vibration",
+			"requirement": "optional"
+		},
+		{
+			"name": "/actions/main/in/VR_Trigger_L",
+			"type": "boolean"
+		},
+		{
+			"name": "/actions/main/in/VR_Trigger_R",
+			"type": "boolean"
+		},
+		{
+			"name": "/actions/main/in/ResetLocationL",
+			"type": "boolean"
+		},
+		{
+			"name": "/actions/main/in/ShowContextualMenu_R",
+			"type": "boolean"
+		},
+		{
+			"name": "/actions/main/in/ShowContextualMenu_L",
+			"type": "boolean"
+		},
+		{
+			"name": "/actions/main/in/VR_Teleport_L",
+			"type": "boolean"
+		},
+		{
+			"name": "/actions/main/in/VR_Teleport_R",
+			"type": "boolean"
+		},
+		{
+			"name": "/actions/main/in/UP_R",
+			"type": "boolean"
+		},
+		{
+			"name": "/actions/main/in/UP_L",
+			"type": "boolean"
+		},
+		{
+			"name": "/actions/main/in/Down_R",
+			"type": "boolean"
+		},
+		{
+			"name": "/actions/main/in/Down_L",
+			"type": "boolean"
+		},
+		{
+			"name": "/actions/main/in/Right_R",
+			"type": "boolean"
+		},
+		{
+			"name": "/actions/main/in/Right_L",
+			"type": "boolean"
+		},
+		{
+			"name": "/actions/main/in/Left_R",
+			"type": "boolean"
+		},
+		{
+			"name": "/actions/main/in/Left_L",
+			"type": "boolean"
+		},
+		{
+			"name": "/actions/main/in/Thumbstick_R_Y axis",
+			"type": "vector1"
+		},
+		{
+			"name": "/actions/main/in/Thumbstick_L_Y axis",
+			"type": "vector1"
+		}
+	],
+	"action_sets": [
+		{
+			"name": "/actions/main",
+			"usage": "leftright"
+		}
+	],
+	"default_bindings": [
+		{
+			"controller_type": "knuckles",
+			"binding_url": "knuckles.json"
+		},
+		{
+			"controller_type": "vive_controller",
+			"binding_url": "vive_controller.json"
+		},
+		{
+			"controller_type": "vive_cosmos_controller",
+			"binding_url": "vive_cosmos_controller.json"
+		},
+		{
+			"controller_type": "oculus_touch",
+			"binding_url": "oculus_touch.json"
+		},
+		{
+			"controller_type": "holographic_controller",
+			"binding_url": "holographic_controller.json"
+		},
+		{
+			"controller_type": "indexhmd",
+			"binding_url": "indexhmd.json"
+		},
+		{
+			"controller_type": "vive",
+			"binding_url": "vive.json"
+		},
+		{
+			"controller_type": "vive_pro",
+			"binding_url": "vive_pro.json"
+		},
+		{
+			"controller_type": "rift",
+			"binding_url": "rift.json"
+		},
+		{
+			"controller_type": "vive_tracker_camera",
+			"binding_url": "vive_tracker_camera.json"
+		},
+		{
+			"controller_type": "gamepad",
+			"binding_url": "gamepad.json"
+		}
+	],
+	"localization": [
+		{
+			"language_tag": "en_us",
+			"/actions/main/in/controllerleft": "Left Controller [Pose]",
+			"/actions/main/in/controllerright": "Right Controller [Pose]",
+			"/actions/main/in/special1": "Special 1 [Tracker]",
+			"/actions/main/in/special2": "Special 2 [Tracker]",
+			"/actions/main/in/special3": "Special 3 [Tracker]",
+			"/actions/main/in/special4": "Special 4 [Tracker]",
+			"/actions/main/in/special5": "Special 5 [Tracker]",
+			"/actions/main/in/special6": "Special 6 [Tracker]",
+			"/actions/main/in/special7": "Special 7 [Tracker]",
+			"/actions/main/in/special8": "Special 8 [Tracker]",
+			"/actions/main/in/skeletonleft": "Skeleton (Left)",
+			"/actions/main/in/skeletonright": "Skeleton (Right)",
+			"/actions/main/in/open_console": "Open Console",
+			"/actions/main/out/vibrateleft": "Haptic (Left)",
+			"/actions/main/out/vibrateright": "Haptic (Right)",
+			"/actions/main/in/VR_Trigger_L": "VR_Trigger_L",
+			"/actions/main/in/VR_Trigger_R": "VR_Trigger_R",
+			"/actions/main/in/ResetLocationL": "ResetLocationL",
+			"/actions/main/in/ShowContextualMenu_R": "ShowContextualMenu_R",
+			"/actions/main/in/ShowContextualMenu_L": "ShowContextualMenu_L",
+			"/actions/main/in/VR_Teleport_L": "VR_Teleport_L",
+			"/actions/main/in/VR_Teleport_R": "VR_Teleport_R",
+			"/actions/main/in/UP_R": "UP_R",
+			"/actions/main/in/UP_L": "UP_L",
+			"/actions/main/in/Down_R": "Down_R",
+			"/actions/main/in/Down_L": "Down_L",
+			"/actions/main/in/Right_R": "Right_R",
+			"/actions/main/in/Right_L": "Right_L",
+			"/actions/main/in/Left_R": "Left_R",
+			"/actions/main/in/Left_L": "Left_L",
+			"/actions/main/in/Thumbstick_R_Y axis": "Thumbstick_R_Y",
+			"/actions/main/in/Thumbstick_L_Y axis": "Thumbstick_L_Y",
+			"/actions/main": "Main Game Actions"
+		}
+	]
+}
\ No newline at end of file
diff --git a/Templates/TP_CollaborativeBP/Config/SteamVRBindings/vive.json b/Templates/TP_CollaborativeBP/Config/SteamVRBindings/vive.json
new file mode 100644
--- /dev/null
+++ b/Templates/TP_CollaborativeBP/Config/SteamVRBindings/vive.json
@@ -0,0 +1,13 @@
+{
+	"name": "Default bindings for Vive Headset",
+	"controller_type": "vive",
+	"last_edited_by": "UnrealEngine",
+	"bindings":
+	{
+		"/actions/main":
+		{
+			"sources": []
+		}
+	},
+	"description": "TP_CollaborativeBP-10452296"
+}
\ No newline at end of file
diff --git a/Templates/TP_CollaborativeBP/Config/SteamVRBindings/vive_controller.json b/Templates/TP_CollaborativeBP/Config/SteamVRBindings/vive_controller.json
new file mode 100644
--- /dev/null
+++ b/Templates/TP_CollaborativeBP/Config/SteamVRBindings/vive_controller.json
@@ -0,0 +1,242 @@
+{
+	"name": "Default bindings for Vive",
+	"controller_type": "vive_controller",
+	"last_edited_by": "UnrealEngine",
+	"bindings":
+	{
+		"/actions/main":
+		{
+			"sources": [
+				{
+					"mode": "button",
+					"path": "/user/hand/left/input/trigger",
+					"inputs":
+					{
+						"click":
+						{
+							"output": "/actions/main/in/VR_Trigger_L"
+						}
+					}
+				},
+				{
+					"mode": "button",
+					"path": "/user/hand/right/input/trigger",
+					"inputs":
+					{
+						"click":
+						{
+							"output": "/actions/main/in/VR_Trigger_R"
+						}
+					}
+				},
+				{
+					"mode": "button",
+					"path": "/user/hand/left/input/grip",
+					"inputs":
+					{
+						"click":
+						{
+							"output": "/actions/main/in/ResetLocationL"
+						}
+					}
+				},
+				{
+					"mode": "button",
+					"path": "/user/hand/right/input/application_menu",
+					"inputs":
+					{
+						"click":
+						{
+							"output": "/actions/main/in/ShowContextualMenu_R"
+						}
+					}
+				},
+				{
+					"mode": "button",
+					"path": "/user/hand/left/input/application_menu",
+					"inputs":
+					{
+						"click":
+						{
+							"output": "/actions/main/in/ShowContextualMenu_L"
+						}
+					}
+				},
+				{
+					"mode": "button",
+					"path": "/user/hand/left/input/grip",
+					"inputs":
+					{
+						"click":
+						{
+							"output": "/actions/main/in/VR_Teleport_L"
+						}
+					}
+				},
+				{
+					"mode": "button",
+					"path": "/user/hand/right/input/grip",
+					"inputs":
+					{
+						"click":
+						{
+							"output": "/actions/main/in/VR_Teleport_R"
+						}
+					}
+				},
+				{
+					"mode": "dpad",
+					"path": "/user/hand/right/input/trackpad",
+					"parameters":
+					{
+						"sub_mode": "click"
+					},
+					"inputs":
+					{
+						"north":
+						{
+							"output": "/actions/main/in/UP_R"
+						}
+					}
+				},
+				{
+					"mode": "dpad",
+					"path": "/user/hand/left/input/trackpad",
+					"parameters":
+					{
+						"sub_mode": "click"
+					},
+					"inputs":
+					{
+						"north":
+						{
+							"output": "/actions/main/in/UP_L"
+						}
+					}
+				},
+				{
+					"mode": "dpad",
+					"path": "/user/hand/right/input/trackpad",
+					"parameters":
+					{
+						"sub_mode": "click"
+					},
+					"inputs":
+					{
+						"south":
+						{
+							"output": "/actions/main/in/Down_R"
+						}
+					}
+				},
+				{
+					"mode": "dpad",
+					"path": "/user/hand/left/input/trackpad",
+					"parameters":
+					{
+						"sub_mode": "click"
+					},
+					"inputs":
+					{
+						"south":
+						{
+							"output": "/actions/main/in/Down_L"
+						}
+					}
+				},
+				{
+					"mode": "dpad",
+					"path": "/user/hand/right/input/trackpad",
+					"parameters":
+					{
+						"sub_mode": "click"
+					},
+					"inputs":
+					{
+						"east":
+						{
+							"output": "/actions/main/in/Right_R"
+						}
+					}
+				},
+				{
+					"mode": "dpad",
+					"path": "/user/hand/left/input/trackpad",
+					"parameters":
+					{
+						"sub_mode": "click"
+					},
+					"inputs":
+					{
+						"east":
+						{
+							"output": "/actions/main/in/Right_L"
+						}
+					}
+				},
+				{
+					"mode": "dpad",
+					"path": "/user/hand/right/input/trackpad",
+					"parameters":
+					{
+						"sub_mode": "click"
+					},
+					"inputs":
+					{
+						"west":
+						{
+							"output": "/actions/main/in/Left_R"
+						}
+					}
+				},
+				{
+					"mode": "dpad",
+					"path": "/user/hand/left/input/trackpad",
+					"parameters":
+					{
+						"sub_mode": "click"
+					},
+					"inputs":
+					{
+						"west":
+						{
+							"output": "/actions/main/in/Left_L"
+						}
+					}
+				}
+			],
+			"poses": [
+				{
+					"output": "/actions/main/in/controllerleft",
+					"path": "/user/hand/left/pose/raw",
+					"requirement": "optional"
+				},
+				{
+					"output": "/actions/main/in/controllerright",
+					"path": "/user/hand/right/pose/raw"
+				}
+			],
+			"skeleton": [
+				{
+					"output": "/actions/main/in/skeletonleft",
+					"path": "/user/hand/left/input/skeleton/left"
+				},
+				{
+					"output": "/actions/main/in/skeletonright",
+					"path": "/user/hand/right/input/skeleton/right"
+				}
+			],
+			"haptics": [
+				{
+					"output": "/actions/main/out/vibrateleft",
+					"path": "/user/hand/left/output/haptic"
+				},
+				{
+					"output": "/actions/main/out/vibrateright",
+					"path": "/user/hand/right/output/haptic"
+				}
+			]
+		}
+	},
+	"description": "TP_CollaborativeBP-10452296"
+}
\ No newline at end of file
diff --git a/Templates/TP_CollaborativeBP/Config/SteamVRBindings/vive_cosmos_controller.json b/Templates/TP_CollaborativeBP/Config/SteamVRBindings/vive_cosmos_controller.json
new file mode 100644
--- /dev/null
+++ b/Templates/TP_CollaborativeBP/Config/SteamVRBindings/vive_cosmos_controller.json
@@ -0,0 +1,44 @@
+{
+	"name": "Default bindings for Cosmos",
+	"controller_type": "vive_cosmos_controller",
+	"last_edited_by": "UnrealEngine",
+	"bindings":
+	{
+		"/actions/main":
+		{
+			"sources": [],
+			"poses": [
+				{
+					"output": "/actions/main/in/controllerleft",
+					"path": "/user/hand/left/pose/raw",
+					"requirement": "optional"
+				},
+				{
+					"output": "/actions/main/in/controllerright",
+					"path": "/user/hand/right/pose/raw"
+				}
+			],
+			"skeleton": [
+				{
+					"output": "/actions/main/in/skeletonleft",
+					"path": "/user/hand/left/input/skeleton/left"
+				},
+				{
+					"output": "/actions/main/in/skeletonright",
+					"path": "/user/hand/right/input/skeleton/right"
+				}
+			],
+			"haptics": [
+				{
+					"output": "/actions/main/out/vibrateleft",
+					"path": "/user/hand/left/output/haptic"
+				},
+				{
+					"output": "/actions/main/out/vibrateright",
+					"path": "/user/hand/right/output/haptic"
+				}
+			]
+		}
+	},
+	"description": "TP_CollaborativeBP-10452296"
+}
\ No newline at end of file
diff --git a/Templates/TP_CollaborativeBP/Config/SteamVRBindings/vive_pro.json b/Templates/TP_CollaborativeBP/Config/SteamVRBindings/vive_pro.json
new file mode 100644
--- /dev/null
+++ b/Templates/TP_CollaborativeBP/Config/SteamVRBindings/vive_pro.json
@@ -0,0 +1,13 @@
+{
+	"name": "Default bindings for Vive Pro Headset",
+	"controller_type": "vive_pro",
+	"last_edited_by": "UnrealEngine",
+	"bindings":
+	{
+		"/actions/main":
+		{
+			"sources": []
+		}
+	},
+	"description": "TP_CollaborativeBP-10452296"
+}
\ No newline at end of file
diff --git a/Templates/TP_CollaborativeBP/Config/SteamVRBindings/vive_tracker_camera.json b/Templates/TP_CollaborativeBP/Config/SteamVRBindings/vive_tracker_camera.json
new file mode 100644
--- /dev/null
+++ b/Templates/TP_CollaborativeBP/Config/SteamVRBindings/vive_tracker_camera.json
@@ -0,0 +1,253 @@
+{
+	"name": "Default bindings for Vive Trackers",
+	"controller_type": "vive_tracker_camera",
+	"last_edited_by": "UnrealEngine",
+	"bindings":
+	{
+		"/actions/main":
+		{
+			"sources": [
+				{
+					"mode": "button",
+					"path": "/user/hand/left/input/trigger",
+					"inputs":
+					{
+						"click":
+						{
+							"output": "/actions/main/in/VR_Trigger_L"
+						}
+					}
+				},
+				{
+					"mode": "button",
+					"path": "/user/hand/right/input/trigger",
+					"inputs":
+					{
+						"click":
+						{
+							"output": "/actions/main/in/VR_Trigger_R"
+						}
+					}
+				},
+				{
+					"mode": "button",
+					"path": "/user/hand/left/input/grip",
+					"inputs":
+					{
+						"click":
+						{
+							"output": "/actions/main/in/ResetLocationL"
+						}
+					}
+				},
+				{
+					"mode": "button",
+					"path": "/user/hand/right/input/application_menu",
+					"inputs":
+					{
+						"click":
+						{
+							"output": "/actions/main/in/ShowContextualMenu_R"
+						}
+					}
+				},
+				{
+					"mode": "button",
+					"path": "/user/hand/left/input/application_menu",
+					"inputs":
+					{
+						"click":
+						{
+							"output": "/actions/main/in/ShowContextualMenu_L"
+						}
+					}
+				},
+				{
+					"mode": "button",
+					"path": "/user/hand/left/input/grip",
+					"inputs":
+					{
+						"click":
+						{
+							"output": "/actions/main/in/VR_Teleport_L"
+						}
+					}
+				},
+				{
+					"mode": "button",
+					"path": "/user/hand/right/input/grip",
+					"inputs":
+					{
+						"click":
+						{
+							"output": "/actions/main/in/VR_Teleport_R"
+						}
+					}
+				},
+				{
+					"mode": "dpad",
+					"path": "/user/hand/right/input/trackpad",
+					"parameters":
+					{
+						"sub_mode": "click"
+					},
+					"inputs":
+					{
+						"north":
+						{
+							"output": "/actions/main/in/UP_R"
+						}
+					}
+				},
+				{
+					"mode": "dpad",
+					"path": "/user/hand/left/input/trackpad",
+					"parameters":
+					{
+						"sub_mode": "click"
+					},
+					"inputs":
+					{
+						"north":
+						{
+							"output": "/actions/main/in/UP_L"
+						}
+					}
+				},
+				{
+					"mode": "dpad",
+					"path": "/user/hand/right/input/trackpad",
+					"parameters":
+					{
+						"sub_mode": "click"
+					},
+					"inputs":
+					{
+						"south":
+						{
+							"output": "/actions/main/in/Down_R"
+						}
+					}
+				},
+				{
+					"mode": "dpad",
+					"path": "/user/hand/left/input/trackpad",
+					"parameters":
+					{
+						"sub_mode": "click"
+					},
+					"inputs":
+					{
+						"south":
+						{
+							"output": "/actions/main/in/Down_L"
+						}
+					}
+				},
+				{
+					"mode": "dpad",
+					"path": "/user/hand/right/input/trackpad",
+					"parameters":
+					{
+						"sub_mode": "click"
+					},
+					"inputs":
+					{
+						"east":
+						{
+							"output": "/actions/main/in/Right_R"
+						}
+					}
+				},
+				{
+					"mode": "dpad",
+					"path": "/user/hand/left/input/trackpad",
+					"parameters":
+					{
+						"sub_mode": "click"
+					},
+					"inputs":
+					{
+						"east":
+						{
+							"output": "/actions/main/in/Right_L"
+						}
+					}
+				},
+				{
+					"mode": "dpad",
+					"path": "/user/hand/right/input/trackpad",
+					"parameters":
+					{
+						"sub_mode": "click"
+					},
+					"inputs":
+					{
+						"west":
+						{
+							"output": "/actions/main/in/Left_R"
+						}
+					}
+				},
+				{
+					"mode": "dpad",
+					"path": "/user/hand/left/input/trackpad",
+					"parameters":
+					{
+						"sub_mode": "click"
+					},
+					"inputs":
+					{
+						"west":
+						{
+							"output": "/actions/main/in/Left_L"
+						}
+					}
+				}
+			],
+			"poses": [
+				{
+					"output": "/actions/main/in/special1",
+					"path": "/user/hand/left/pose/back",
+					"requirement": "optional"
+				},
+				{
+					"output": "/actions/main/in/special2",
+					"path": "/user/hand/right/pose/back",
+					"requirement": "optional"
+				},
+				{
+					"output": "/actions/main/in/special3",
+					"path": "/user/hand/left/pose/front",
+					"requirement": "optional"
+				},
+				{
+					"output": "/actions/main/in/special4",
+					"path": "/user/hand/right/pose/front",
+					"requirement": "optional"
+				},
+				{
+					"output": "/actions/main/in/special5",
+					"path": "/user/hand/left/pose/frontandrolled",
+					"requirement": "optional"
+				},
+				{
+					"output": "/actions/main/in/special6",
+					"path": "/user/hand/right/pose/frontandrolled",
+					"requirement": "optional"
+				},
+				{
+					"output": "/actions/main/in/special7",
+					"path": "/user/hand/left/pose/pistolgrip",
+					"requirement": "optional"
+				},
+				{
+					"output": "/actions/main/in/special8",
+					"path": "/user/hand/right/pose/pistolgrip",
+					"requirement": "optional"
+				}
+			]
+		}
+	},
+	"description": "TP_CollaborativeBP-10452296"
+}
\ No newline at end of file
diff --git a/Templates/TP_CollaborativeBP/Config/steamvr_ue_editor_app.json b/Templates/TP_CollaborativeBP/Config/steamvr_ue_editor_app.json
new file mode 100644
--- /dev/null
+++ b/Templates/TP_CollaborativeBP/Config/steamvr_ue_editor_app.json
@@ -0,0 +1,18 @@
+{
+	"source": "UE",
+	"applications": [
+		{
+			"app_key": "application.generated.ue.tp_collaborativebp-10452296.ue4editor-win64-debuggame.exe",
+			"launch_type": "url",
+			"url": "steam://launch/",
+			"action_manifest_path": "D:/Project/424Full/Templates/TP_CollaborativeBP/Config/SteamVRBindings/steamvr_manifest.json",
+			"strings":
+			{
+				"en_us":
+				{
+					"name": "TP_CollaborativeBP-10452296 [UE Editor]"
+				}
+			}
+		}
+	]
+}
\ No newline at end of file
