-Fix Client server issues
-Optimize Memory size of the Collab ViewerTemplate ( Tree )

#rnx

diff --git a/Templates/TP_CollaborativeBP/Config/DefaultInput.ini b/Templates/TP_CollaborativeBP/Config/DefaultInput.ini
--- a/Templates/TP_CollaborativeBP/Config/DefaultInput.ini
+++ b/Templates/TP_CollaborativeBP/Config/DefaultInput.ini
@@ -1,5 +1,11 @@
 
 [/Script/Engine.InputSettings]
++AxisConfig=(AxisKeyName="MixedReality_Left_Thumbstick_X",AxisProperties=(DeadZone=0.000000,Sensitivity=1.000000,Exponent=1.000000,bInvert=False))
++AxisConfig=(AxisKeyName="MixedReality_Left_Thumbstick_Y",AxisProperties=(DeadZone=0.000000,Sensitivity=1.000000,Exponent=1.000000,bInvert=False))
++AxisConfig=(AxisKeyName="MixedReality_Left_Trackpad_X",AxisProperties=(DeadZone=0.000000,Sensitivity=1.000000,Exponent=1.000000,bInvert=False))
++AxisConfig=(AxisKeyName="MixedReality_Left_Trackpad_Y",AxisProperties=(DeadZone=0.000000,Sensitivity=1.000000,Exponent=1.000000,bInvert=False))
++AxisConfig=(AxisKeyName="MixedReality_Right_Trigger_Axis",AxisProperties=(DeadZone=0.000000,Sensitivity=1.000000,Exponent=1.000000,bInvert=False))
++AxisConfig=(AxisKeyName="MixedReality_Right_Thumbstick_X",AxisProperties=(DeadZone=0.000000,Sensitivity=1.000000,Exponent=1.000000,bInvert=False))
 +AxisConfig=(AxisKeyName="MixedReality_Right_Thumbstick_Y",AxisProperties=(DeadZone=0.000000,Sensitivity=1.000000,Exponent=1.000000,bInvert=False))
 +AxisConfig=(AxisKeyName="MixedReality_Right_Trackpad_X",AxisProperties=(DeadZone=0.000000,Sensitivity=1.000000,Exponent=1.000000,bInvert=False))
 +AxisConfig=(AxisKeyName="MixedReality_Right_Trackpad_Y",AxisProperties=(DeadZone=0.000000,Sensitivity=1.000000,Exponent=1.000000,bInvert=False))
@@ -112,12 +118,6 @@
 +AxisConfig=(AxisKeyName="Vive_Right_Trackpad_X",AxisProperties=(DeadZone=0.000000,Sensitivity=1.000000,Exponent=1.000000,bInvert=False))
 +AxisConfig=(AxisKeyName="Vive_Right_Trackpad_Y",AxisProperties=(DeadZone=0.000000,Sensitivity=1.000000,Exponent=1.000000,bInvert=False))
 +AxisConfig=(AxisKeyName="MixedReality_Left_Trigger_Axis",AxisProperties=(DeadZone=0.000000,Sensitivity=1.000000,Exponent=1.000000,bInvert=False))
-+AxisConfig=(AxisKeyName="MixedReality_Left_Thumbstick_X",AxisProperties=(DeadZone=0.000000,Sensitivity=1.000000,Exponent=1.000000,bInvert=False))
-+AxisConfig=(AxisKeyName="MixedReality_Left_Thumbstick_Y",AxisProperties=(DeadZone=0.000000,Sensitivity=1.000000,Exponent=1.000000,bInvert=False))
-+AxisConfig=(AxisKeyName="MixedReality_Left_Trackpad_X",AxisProperties=(DeadZone=0.000000,Sensitivity=1.000000,Exponent=1.000000,bInvert=False))
-+AxisConfig=(AxisKeyName="MixedReality_Left_Trackpad_Y",AxisProperties=(DeadZone=0.000000,Sensitivity=1.000000,Exponent=1.000000,bInvert=False))
-+AxisConfig=(AxisKeyName="MixedReality_Right_Trigger_Axis",AxisProperties=(DeadZone=0.000000,Sensitivity=1.000000,Exponent=1.000000,bInvert=False))
-+AxisConfig=(AxisKeyName="MixedReality_Right_Thumbstick_X",AxisProperties=(DeadZone=0.000000,Sensitivity=1.000000,Exponent=1.000000,bInvert=False))
 bAltEnterTogglesFullscreen=True
 bF11TogglesFullscreen=True
 bUseMouseForTouch=False
@@ -217,9 +217,9 @@ DoubleClickTime=0.200000
 +AxisMappings=(AxisName="WalkForward",Scale=-1.000000,Key=Down)
 +AxisMappings=(AxisName="WalkRight",Scale=1.000000,Key=Right)
 +AxisMappings=(AxisName="WalkRight",Scale=-1.000000,Key=Left)
-+AxisMappings=(AxisName="Thumbstick_R_Y",Scale=1.000000,Key=OculusTouch_Right_Thumbstick_Y)
++AxisMappings=(AxisName="Thumbstick_R_Y",Scale=-1.000000,Key=OculusTouch_Right_Thumbstick_Y)
 +AxisMappings=(AxisName="Thumbstick_L_Y",Scale=1.000000,Key=OculusTouch_Left_Thumbstick_Y)
-+AxisMappings=(AxisName="Thumbstick_R_Y",Scale=1.000000,Key=ValveIndex_Right_Thumbstick_Y)
++AxisMappings=(AxisName="Thumbstick_R_Y",Scale=-1.000000,Key=ValveIndex_Right_Thumbstick_Y)
 +AxisMappings=(AxisName="Thumbstick_L_Y",Scale=1.000000,Key=ValveIndex_Left_Thumbstick_Y)
 DefaultTouchInterface=/Game/CollaborativeViewer/Blueprints/Pawn/Touch/TouchInterface.TouchInterface
 +ConsoleKeys=Tilde
diff --git a/Templates/TP_CollaborativeBP/Config/TemplateDefs.ini b/Templates/TP_CollaborativeBP/Config/TemplateDefs.ini
--- a/Templates/TP_CollaborativeBP/Config/TemplateDefs.ini
+++ b/Templates/TP_CollaborativeBP/Config/TemplateDefs.ini
@@ -13,10 +13,6 @@ LocalizedDisplayNames=(Language="en", Text="Collab Viewer")
 
 ; The localized descriptions for this template, which is displayed when the template is selected.
 LocalizedDescriptions=(Language="en", Text="This project template provides navigation and interaction for various platforms (Desktop, VR). Use as a starting point for collaborative industrial or architectural projects.\n\nThis template enables you to use Datasmith to import models from various design applications and easily light them.")
-
-LocalizedDisplayNames=(Language="ko", Text="콜라보 뷰어")
-LocalizedDescriptions=(Language="ko", Text="이 프로젝트 템플릿은 데스크탑, VR등 다양한 플랫폼을 위한 네비게이션과 인터렉션을 제공합니다. 산업용 협업 또는 건축 프로젝트를 위한 시작 프로젝트로 사용하세요.")
-
 ; The sort key to use for the template to bypass the default alphabetical sorting.
 SortKey="3"
 
diff --git a/Templates/TP_CollaborativeBP/Config/steamvr_ue_editor_app.json b/Templates/TP_CollaborativeBP/Config/steamvr_ue_editor_app.json
--- a/Templates/TP_CollaborativeBP/Config/steamvr_ue_editor_app.json
+++ b/Templates/TP_CollaborativeBP/Config/steamvr_ue_editor_app.json
@@ -2,7 +2,7 @@
 	"source": "UE",
 	"applications": [
 		{
-			"app_key": "application.generated.ue.tp_collaborativebp-10452296.ue4editor-win64-debuggame.exe",
+			"app_key": "application.generated.ue.tp_collaborativebp-10955000.ue4editor-win64-debuggame.exe",
 			"launch_type": "url",
 			"url": "steam://launch/",
 			"action_manifest_path": "D:/Project/424Full/Templates/TP_CollaborativeBP/Config/SteamVRBindings/steamvr_manifest.json",
@@ -10,7 +10,7 @@
 			{
 				"en_us":
 				{
-					"name": "TP_CollaborativeBP-10452296 [UE Editor]"
+					"name": "TP_CollaborativeBP-10955000 [UE Editor]"
 				}
 			}
 		}
diff --git a/Templates/TP_CollaborativeBP/TP_CollaborativeBP.uproject b/Templates/TP_CollaborativeBP/TP_CollaborativeBP.uproject
--- a/Templates/TP_CollaborativeBP/TP_CollaborativeBP.uproject
+++ b/Templates/TP_CollaborativeBP/TP_CollaborativeBP.uproject
@@ -1,6 +1,6 @@
 {
 	"FileVersion": 3,
-	"EngineAssociation": "",
+	"EngineAssociation": "4.24",
 	"Category": "",
 	"Description": "",
 	"Plugins": [
diff --git a/Templates/TP_CollaborativeBP/manifest.json b/Templates/TP_CollaborativeBP/manifest.json
--- a/Templates/TP_CollaborativeBP/manifest.json
+++ b/Templates/TP_CollaborativeBP/manifest.json
@@ -5,10 +5,6 @@
 		{
 			"Language": "en",
 			"Text": "Collab Viewer"
-		},
-		{
-			"Language": "ko",
-			"Text": "콜라보 뷰어"
 		}
 	],
 	"Description":
@@ -16,10 +12,6 @@
 		{
 			"Language": "en",
 			"Text": "This project template provides navigation and interaction for various platforms (Desktop, VR). Use as a starting point for collaborative Industrial or Architectural projects."
-		},
-		{
-			"Language": "ko",
-			"Text": "이 프로젝트 템플릿은 데스크탑, VR등 다양한 플랫폼을 위한 네비게이션과 인터렉션을 제공합니다. 산업용 협업 또는 건축 프로젝트를 위한 시작 프로젝트로 사용하세요."
 		}
 	],
 	"AssetTypes": [],
