-   Hollolens Support for Windows Only, Fix Camera Warping

diff --git a/Templates/TP_CollaborativeBP/Config/DefaultGame.ini b/Templates/TP_CollaborativeBP/Config/DefaultGame.ini
--- a/Templates/TP_CollaborativeBP/Config/DefaultGame.ini
+++ b/Templates/TP_CollaborativeBP/Config/DefaultGame.ini
@@ -5,3 +5,6 @@
 ProjectID=C684876B46D5306ADA6286AA20FA1808
 bStartInVR=False
 
+[/Script/UnrealEd.ProjectPackagingSettings]
+bSkipEditorContent=False
+
diff --git a/Templates/TP_CollaborativeBP/Config/DefaultInput.ini b/Templates/TP_CollaborativeBP/Config/DefaultInput.ini
--- a/Templates/TP_CollaborativeBP/Config/DefaultInput.ini
+++ b/Templates/TP_CollaborativeBP/Config/DefaultInput.ini
@@ -1,5 +1,18 @@
 
+
 [/Script/Engine.InputSettings]
++AxisConfig=(AxisKeyName="WindowsSpatialInput_RightManipulationXGesture",AxisProperties=(DeadZone=0.000000,Sensitivity=1.000000,Exponent=1.000000,bInvert=False))
++AxisConfig=(AxisKeyName="WindowsSpatialInput_RightManipulationYGesture",AxisProperties=(DeadZone=0.000000,Sensitivity=1.000000,Exponent=1.000000,bInvert=False))
++AxisConfig=(AxisKeyName="WindowsSpatialInput_RightManipulationZGesture",AxisProperties=(DeadZone=0.000000,Sensitivity=1.000000,Exponent=1.000000,bInvert=False))
++AxisConfig=(AxisKeyName="WindowsSpatialInput_RightNavigationXGesture",AxisProperties=(DeadZone=0.000000,Sensitivity=1.000000,Exponent=1.000000,bInvert=False))
++AxisConfig=(AxisKeyName="WindowsSpatialInput_RightNavigationYGesture",AxisProperties=(DeadZone=0.000000,Sensitivity=1.000000,Exponent=1.000000,bInvert=False))
++AxisConfig=(AxisKeyName="WindowsSpatialInput_RightNavigationZGesture",AxisProperties=(DeadZone=0.000000,Sensitivity=1.000000,Exponent=1.000000,bInvert=False))
++AxisConfig=(AxisKeyName="Vive_Left_Trackpad_X",AxisProperties=(DeadZone=0.000000,Sensitivity=1.000000,Exponent=1.000000,bInvert=False))
++AxisConfig=(AxisKeyName="Vive_Left_Trackpad_Y",AxisProperties=(DeadZone=0.000000,Sensitivity=1.000000,Exponent=1.000000,bInvert=False))
++AxisConfig=(AxisKeyName="Vive_Right_Trigger_Axis",AxisProperties=(DeadZone=0.000000,Sensitivity=1.000000,Exponent=1.000000,bInvert=False))
++AxisConfig=(AxisKeyName="Vive_Right_Trackpad_X",AxisProperties=(DeadZone=0.000000,Sensitivity=1.000000,Exponent=1.000000,bInvert=False))
++AxisConfig=(AxisKeyName="Vive_Right_Trackpad_Y",AxisProperties=(DeadZone=0.000000,Sensitivity=1.000000,Exponent=1.000000,bInvert=False))
++AxisConfig=(AxisKeyName="MixedReality_Left_Trigger_Axis",AxisProperties=(DeadZone=0.000000,Sensitivity=1.000000,Exponent=1.000000,bInvert=False))
 +AxisConfig=(AxisKeyName="MixedReality_Left_Thumbstick_X",AxisProperties=(DeadZone=0.000000,Sensitivity=1.000000,Exponent=1.000000,bInvert=False))
 +AxisConfig=(AxisKeyName="MixedReality_Left_Thumbstick_Y",AxisProperties=(DeadZone=0.000000,Sensitivity=1.000000,Exponent=1.000000,bInvert=False))
 +AxisConfig=(AxisKeyName="MixedReality_Left_Trackpad_X",AxisProperties=(DeadZone=0.000000,Sensitivity=1.000000,Exponent=1.000000,bInvert=False))
@@ -112,12 +125,12 @@
 +AxisConfig=(AxisKeyName="Daydream_Right_Trackpad_X",AxisProperties=(DeadZone=0.000000,Sensitivity=1.000000,Exponent=1.000000,bInvert=False))
 +AxisConfig=(AxisKeyName="Daydream_Right_Trackpad_Y",AxisProperties=(DeadZone=0.000000,Sensitivity=1.000000,Exponent=1.000000,bInvert=False))
 +AxisConfig=(AxisKeyName="Vive_Left_Trigger_Axis",AxisProperties=(DeadZone=0.000000,Sensitivity=1.000000,Exponent=1.000000,bInvert=False))
-+AxisConfig=(AxisKeyName="Vive_Left_Trackpad_X",AxisProperties=(DeadZone=0.000000,Sensitivity=1.000000,Exponent=1.000000,bInvert=False))
-+AxisConfig=(AxisKeyName="Vive_Left_Trackpad_Y",AxisProperties=(DeadZone=0.000000,Sensitivity=1.000000,Exponent=1.000000,bInvert=False))
-+AxisConfig=(AxisKeyName="Vive_Right_Trigger_Axis",AxisProperties=(DeadZone=0.000000,Sensitivity=1.000000,Exponent=1.000000,bInvert=False))
-+AxisConfig=(AxisKeyName="Vive_Right_Trackpad_X",AxisProperties=(DeadZone=0.000000,Sensitivity=1.000000,Exponent=1.000000,bInvert=False))
-+AxisConfig=(AxisKeyName="Vive_Right_Trackpad_Y",AxisProperties=(DeadZone=0.000000,Sensitivity=1.000000,Exponent=1.000000,bInvert=False))
-+AxisConfig=(AxisKeyName="MixedReality_Left_Trigger_Axis",AxisProperties=(DeadZone=0.000000,Sensitivity=1.000000,Exponent=1.000000,bInvert=False))
++AxisConfig=(AxisKeyName="WindowsSpatialInput_LeftManipulationXGesture",AxisProperties=(DeadZone=0.000000,Sensitivity=1.000000,Exponent=1.000000,bInvert=False))
++AxisConfig=(AxisKeyName="WindowsSpatialInput_LeftManipulationYGesture",AxisProperties=(DeadZone=0.000000,Sensitivity=1.000000,Exponent=1.000000,bInvert=False))
++AxisConfig=(AxisKeyName="WindowsSpatialInput_LeftManipulationZGesture",AxisProperties=(DeadZone=0.000000,Sensitivity=1.000000,Exponent=1.000000,bInvert=False))
++AxisConfig=(AxisKeyName="WindowsSpatialInput_LeftNavigationXGesture",AxisProperties=(DeadZone=0.000000,Sensitivity=1.000000,Exponent=1.000000,bInvert=False))
++AxisConfig=(AxisKeyName="WindowsSpatialInput_LeftNavigationYGesture",AxisProperties=(DeadZone=0.000000,Sensitivity=1.000000,Exponent=1.000000,bInvert=False))
++AxisConfig=(AxisKeyName="WindowsSpatialInput_LeftNavigationZGesture",AxisProperties=(DeadZone=0.000000,Sensitivity=1.000000,Exponent=1.000000,bInvert=False))
 bAltEnterTogglesFullscreen=True
 bF11TogglesFullscreen=True
 bUseMouseForTouch=False
@@ -204,13 +217,11 @@ DoubleClickTime=0.200000
 +AxisMappings=(AxisName="WalkRight",Scale=1.000000,Key=D)
 +AxisMappings=(AxisName="WalkForward",Scale=-1.000000,Key=S)
 +AxisMappings=(AxisName="WalkRight",Scale=-1.000000,Key=A)
-+AxisMappings=(AxisName="LookUp",Scale=1.000000,Key=MouseY)
 +AxisMappings=(AxisName="LookRight",Scale=1.000000,Key=MouseX)
 +AxisMappings=(AxisName="FlyUp",Scale=1.000000,Key=E)
 +AxisMappings=(AxisName="FlyUp",Scale=-1.000000,Key=Q)
 +AxisMappings=(AxisName="WalkForward",Scale=1.000000,Key=Gamepad_LeftY)
 +AxisMappings=(AxisName="WalkRight",Scale=1.000000,Key=Gamepad_LeftX)
-+AxisMappings=(AxisName="LookUp",Scale=-1.000000,Key=Gamepad_RightY)
 +AxisMappings=(AxisName="LookRight",Scale=1.000000,Key=Gamepad_RightX)
 +AxisMappings=(AxisName="pinch",Scale=1.000000,Key=Gesture_Pinch)
 +AxisMappings=(AxisName="WalkForward",Scale=1.000000,Key=Up)
@@ -221,7 +232,8 @@ DoubleClickTime=0.200000
 +AxisMappings=(AxisName="Thumbstick_L_Y",Scale=1.000000,Key=OculusTouch_Left_Thumbstick_Y)
 +AxisMappings=(AxisName="Thumbstick_R_Y",Scale=-1.000000,Key=ValveIndex_Right_Thumbstick_Y)
 +AxisMappings=(AxisName="Thumbstick_L_Y",Scale=1.000000,Key=ValveIndex_Left_Thumbstick_Y)
++AxisMappings=(AxisName="LookUp",Scale=1.000000,Key=MouseY)
++AxisMappings=(AxisName="LookUp",Scale=1.000000,Key=Gamepad_RightY)
 DefaultTouchInterface=/Game/CollaborativeViewer/Blueprints/Pawn/Touch/TouchInterface.TouchInterface
 +ConsoleKeys=Tilde
 
-
diff --git a/Templates/TP_CollaborativeBP/TP_CollaborativeBP.uproject b/Templates/TP_CollaborativeBP/TP_CollaborativeBP.uproject
--- a/Templates/TP_CollaborativeBP/TP_CollaborativeBP.uproject
+++ b/Templates/TP_CollaborativeBP/TP_CollaborativeBP.uproject
@@ -60,6 +60,22 @@
 			"SupportedTargetPlatforms": [
 				"Win64"
 			]
+		},
+		{
+			"Name": "HoloLensAR",
+			"Enabled": true,
+			"SupportedTargetPlatforms": [
+				"HoloLens",
+				"Win64"
+			]
+		},
+		{
+			"Name": "WindowsMixedReality",
+			"Enabled": true,
+			"SupportedTargetPlatforms": [
+				"Win64",
+				"HoloLens"
+			]
 		}
 	]
 }
\ No newline at end of file
