Make Shared resources for Collaborative viewer and Hololens Viewer Templates.

diff --git a/Templates/TP_CollaborativeBP/Config/DefaultEngine.ini b/Templates/TP_CollaborativeBP/Config/DefaultEngine.ini
--- a/Templates/TP_CollaborativeBP/Config/DefaultEngine.ini
+++ b/Templates/TP_CollaborativeBP/Config/DefaultEngine.ini
@@ -13,6 +13,8 @@ EditorStartupMap=/Game/CollaborativeViewer/Maps/CollaborativeViewer_P.Collaborat
 GameDefaultMap=/Game/CollaborativeViewer/Maps/MainMenu.MainMenu
 GlobalDefaultGameMode=/Game/CollaborativeViewer/Blueprints/GameMode/BP_CollaborativeViewer_GameMode.BP_CollaborativeViewer_GameMode_C
 GameInstanceClass=/Game/CollaborativeViewer/Blueprints/GameInstance/BP_CollaborativeViewer_GameInstance.BP_CollaborativeViewer_GameInstance_C
+ServerDefaultMap=/Game/Maps/MainMenu.MainMenu
+TransitionMap=/Game/Maps/MainMenu.MainMenu
 
 [/Script/Engine.CollisionProfile]
 +Profiles=(Name="NoCollision",CollisionEnabled=NoCollision,ObjectTypeName="WorldStatic",CustomResponses=((Channel="Visibility",Response=ECR_Ignore),(Channel="Camera",Response=ECR_Ignore)),HelpMessage="No collision",bCanModify=False)
@@ -96,4 +98,3 @@ PhysXTreeRebuildRate=10
 DefaultBroadphaseSettings=(bUseMBPOnClient=False,bUseMBPOnServer=False,bUseMBPOuterBounds=False,MBPBounds=(Min=(X=0.000000,Y=0.000000,Z=0.000000),Max=(X=0.000000,Y=0.000000,Z=0.000000),IsValid=0),MBPOuterBounds=(Min=(X=0.000000,Y=0.000000,Z=0.000000),Max=(X=0.000000,Y=0.000000,Z=0.000000),IsValid=0),MBPNumSubdivs=2)
 ChaosSettings=(DefaultThreadingModel=TaskGraph,DedicatedThreadTickMode=VariableCappedWithTarget,DedicatedThreadBufferMode=Double)
 
-
diff --git a/Templates/TP_CollaborativeBP/Config/TemplateDefs.ini b/Templates/TP_CollaborativeBP/Config/TemplateDefs.ini
--- a/Templates/TP_CollaborativeBP/Config/TemplateDefs.ini
+++ b/Templates/TP_CollaborativeBP/Config/TemplateDefs.ini
@@ -12,7 +12,7 @@ Categories=MFG
 LocalizedDisplayNames=(Language="en", Text="Collab Viewer")
 
 ; The localized descriptions for this template, which is displayed when the template is selected.
-LocalizedDescriptions=(Language="en", Text="This project template provides navigation and interaction for various platforms (Desktop, VR). Use as a starting point for collaborative industrial or architectural projects.\n\nThis template enables you to use Datasmith to import models from various design applications and easily light them.")
+LocalizedDescriptions=(Language="en", Text="This project template provides navigation and interaction for Hololens. Use as a starting point for industrial or architectural projects.\n\nThis template enables you to use Datasmith to import models from various design applications and easily light them.")
 ; The sort key to use for the template to bypass the default alphabetical sorting.
 SortKey="3"
 
@@ -94,4 +94,4 @@ HiddenSettings=XR
 ;AssetTypes=Skeleton, Static Mesh, Texture
 
 ; Shared feature packs. The files in these packs listed in these structures marked as 'additionalfiles' will be copied on project generation.
-;SharedContentPacks=(MountName="2DSideScroller",DetailLevels=("Standard"))
+ SharedContentPacks=(MountName="CollaborativeViewer",DetailLevels=("Standard"))
diff --git a/Templates/TP_CollaborativeBP/TP_CollaborativeBP.uproject b/Templates/TP_CollaborativeBP/TP_CollaborativeBP.uproject
--- a/Templates/TP_CollaborativeBP/TP_CollaborativeBP.uproject
+++ b/Templates/TP_CollaborativeBP/TP_CollaborativeBP.uproject
@@ -1,6 +1,6 @@
 {
 	"FileVersion": 3,
-	"EngineAssociation": "4.24",
+	"EngineAssociation": "",
 	"Category": "",
 	"Description": "",
 	"Plugins": [
@@ -60,22 +60,6 @@
 			"SupportedTargetPlatforms": [
 				"Win64"
 			]
-		},
-		{
-			"Name": "HoloLensAR",
-			"Enabled": true,
-			"SupportedTargetPlatforms": [
-				"HoloLens",
-				"Win64"
-			]
-		},
-		{
-			"Name": "WindowsMixedReality",
-			"Enabled": true,
-			"SupportedTargetPlatforms": [
-				"Win64",
-				"HoloLens"
-			]
 		}
 	]
 }
\ No newline at end of file
diff --git a/Templates/TemplateResources/Standard/CollaborativeViewer/FeaturePack/manifest.json b/Templates/TemplateResources/Standard/CollaborativeViewer/FeaturePack/manifest.json
new file mode 100644
--- /dev/null
+++ b/Templates/TemplateResources/Standard/CollaborativeViewer/FeaturePack/manifest.json
@@ -0,0 +1,48 @@
+{
+	"Version": 1,
+	"Name":
+	[
+		{
+			"Language": "en",
+			"Text": "CollaborativeViewer Shared Resource Standard"
+		}
+	],
+	"Description":
+	[
+		{
+			"Language": "en",
+			"Text": "Shared Resource Pack for CollaborativeViewer"
+		}
+	],
+	"AssetTypes":
+	[
+		{
+			"Language": "en",
+			"Text": ""
+		}
+	],
+	"SearchTags":
+	[
+		{
+			"Language": "en",
+			"Text": ""
+		}
+	],	
+	"ClassTypes": " ",
+	"Category": "SharedPack",
+	"Thumbnail": "CollaborativeViewer.png",
+	"Screenshots":
+	[
+		"CollaborativeViewer_Preview.png"
+	],
+	"Ident": "CollaborativeViewerStandard",
+	"AdditionalFiles":
+	{
+		"DestinationFilesFolder": "CollaborativeViewer",
+		"AdditionalFilesList":
+		[
+			"Templates/TemplateResources/Standard/CollaborativeViewer/Content/*.*"
+		]
+	}
+}
+
