Adding Japanese translation #jira ueloc-4712 #fyi michelle.zhao, yali.he

diff --git a/Templates/TP_AEC_ArchvisBP/Config/TemplateDefs.ini b/Templates/TP_AEC_ArchvisBP/Config/TemplateDefs.ini
--- a/Templates/TP_AEC_ArchvisBP/Config/TemplateDefs.ini
+++ b/Templates/TP_AEC_ArchvisBP/Config/TemplateDefs.ini
@@ -19,6 +19,8 @@ LocalizedDescriptions=(Language="ko",Text="이 템플릿은 조명 연구, 인
 LocalizedDisplayNames=(Language="zh-Hans",Text="建筑可视化")
 LocalizedDescriptions=(Language="zh-Hans",Text="该模板专用于建筑可视化工作流程，包含太阳研究、内饰渲染及非照片级风格化渲染示例。\n\n该模板允许使用Datasmith从多种设计软件中导入模型并轻松为其设置光照。")
 
+LocalizedDescriptions=(Language="ja",Text="建築ビジュアライゼーションのワークフローに特化したテンプレートです。日照調査やインテリア レンダリング、ノンフォトリアリスティックで様式化されたレンダリングのためのサンプルが備わっています。\n\nこのテンプレートを利用すると、Datasmith を使ってモデルをさまざまなデザイン アプリケーションからインポートできるようになるとともに、それらを簡単にライティングできるようになります。")
+
 
 ; The sort key to use for the template to bypass the default alphabetical sorting.
 SortKey="1"
diff --git a/Templates/TP_CollaborativeBP/Config/TemplateDefs.ini b/Templates/TP_CollaborativeBP/Config/TemplateDefs.ini
--- a/Templates/TP_CollaborativeBP/Config/TemplateDefs.ini
+++ b/Templates/TP_CollaborativeBP/Config/TemplateDefs.ini
@@ -15,6 +15,7 @@ LocalizedDisplayNames=(Language="zh-Hans", Text="协作查看器")
 ; The localized descriptions for this template, which is displayed when the template is selected.
 LocalizedDescriptions=(Language="en", Text="This project template provides navigation and interaction for Hololens. Use as a starting point for industrial or architectural projects.\n\nThis template enables you to use Datasmith to import models from various design applications and easily light them.")
 LocalizedDescriptions=(Language="zh-Hans", Text="该项目模板提供面向Hololens的导航及互动功能，适合用于工业或建筑项目。\n\n该模板允许使用Datasmith从多种设计软件中导入模型并轻松为其设置光照。")
+LocalizedDescriptions=(Language="ja", Text="このプロジェクト テンプレートでは、Hololens のためのナビゲーションとインタラクションが提供されています。産業プロジェクトまたは建築プロジェクトの出発点として利用できます。\n\nこのテンプレートを利用すると、Datasmith を使ってモデルをさまざまなデザイン アプリケーションからインポートできるようになるとともに、それらを簡単にライティングできるようになります。")
 
 ; The sort key to use for the template to bypass the default alphabetical sorting.
 SortKey="3"
diff --git a/Templates/TP_HololensBP/Config/TemplateDefs.ini b/Templates/TP_HololensBP/Config/TemplateDefs.ini
--- a/Templates/TP_HololensBP/Config/TemplateDefs.ini
+++ b/Templates/TP_HololensBP/Config/TemplateDefs.ini
@@ -15,6 +15,7 @@ LocalizedDisplayNames=(Language="zh-Hans", Text="Hololens查看器")
 ; The localized descriptions for this template, which is displayed when the template is selected.
 LocalizedDescriptions=(Language="en", Text="This project template provides navigation and interaction for Hololens. Use as a starting point for collaborative industrial or architectural projects.\n\nThis template enables you to use Datasmith to import models from various design applications and easily light them.")
 LocalizedDescriptions=(Language="zh-Hans", Text="该项目模板提供面向Hololens的导航及互动功能，适合用于协作式工业或建筑项目。\n\n该模板允许使用Datasmith从多种设计软件中导入模型并轻松为其设置光照。")
+LocalizedDescriptions=(Language="ja", Text="このプロジェクト テンプレートでは、Hololens のためのナビゲーションとインタラクションが提供されています。共同作業による産業プロジェクトまたは建築プロジェクトの出発点として利用できます。\n\nこのテンプレートを利用すると、Datasmith を使ってモデルをさまざまなデザイン アプリケーションからインポートできるようになるとともに、それらを簡単にライティングできるようになります。")
 
 
 ; The sort key to use for the template to bypass the default alphabetical sorting.
diff --git a/Templates/TP_PhotoStudioBP/Config/TemplateDefs.ini b/Templates/TP_PhotoStudioBP/Config/TemplateDefs.ini
--- a/Templates/TP_PhotoStudioBP/Config/TemplateDefs.ini
+++ b/Templates/TP_PhotoStudioBP/Config/TemplateDefs.ini
@@ -19,6 +19,8 @@ LocalizedDescriptions=(Language="ko", Text="이 템플릿은 사진 스튜디오
 LocalizedDisplayNames=(Language="zh-Hans", Text="摄影棚")
 LocalizedDescriptions=(Language="zh-Hans", Text="该模板可以让你在摄影棚环境中创建产品演示和动画。\n\n该模板允许使用Datasmith从多种设计软件中导入模型并轻松为其设置光照。")
 
+LocalizedDescriptions=(Language="ja", Text="このテンプレートでは、製品のプレゼンテーションおよびシネマティックスをフォトスタジオ環境で制作できます。\n\nこのテンプレートを利用すると、Datasmith を使ってモデルをさまざまなデザイン アプリケーションからインポートできるようになるとともに、それらを簡単にライティングできるようになります。")
+
 ; The sort key to use for the template to bypass the default alphabetical sorting.
 SortKey="1"
 
