Fix "Collab Viewer" template description in 4.25

diff --git a/Templates/TP_CollaborativeBP/Config/TemplateDefs.ini b/Templates/TP_CollaborativeBP/Config/TemplateDefs.ini
--- a/Templates/TP_CollaborativeBP/Config/TemplateDefs.ini
+++ b/Templates/TP_CollaborativeBP/Config/TemplateDefs.ini
@@ -14,7 +14,7 @@ LocalizedDisplayNames=(Language="zh-Hans", Text="协作查看器")
 LocalizedDisplayNames=(Language="ko", Text="콜라보 뷰어")
 
 ; The localized descriptions for this template, which is displayed when the template is selected.
-LocalizedDescriptions=(Language="en", Text="This project template provides navigation and interaction for Hololens. Use as a starting point for industrial or architectural projects.\n\nThis template enables you to use Datasmith to import models from various design applications and easily light them.")
+LocalizedDescriptions=(Language="en", Text="This project template provides navigation and interaction for various platforms (Desktop, VR). Use as a starting point for collaborative industrial or architectural projects.\n\nThis template enables you to use Datasmith to import models from various design applications and easily light them")
 LocalizedDescriptions=(Language="zh-Hans", Text="该项目模板提供面向Hololens的导航及互动功能，适合用于工业或建筑项目。\n\n该模板允许使用Datasmith从多种设计软件中导入模型并轻松为其设置光照。")
 LocalizedDescriptions=(Language="ja", Text="このプロジェクト テンプレートでは、Hololens のためのナビゲーションとインタラクションが提供されています。産業プロジェクトまたは建築プロジェクトの出発点として利用できます。\n\nこのテンプレートを利用すると、Datasmith を使ってモデルをさまざまなデザイン アプリケーションからインポートできるようになるとともに、それらを簡単にライティングできるようになります。")
 LocalizedDescriptions=(Language="ko", Text="이 프로젝트 템플릿은 데스크탑, VR등 다양한 플랫폼을 위한 네비게이션과 인터렉션을 제공합니다. 산업용 협업 또는 건축 프로젝트를 위한 시작 프로젝트로 사용하세요.")
